Merging //UE4/Main@9388041 to Dev-RenderPlat-Staging

diff --git a/Engine/Plugins/Experimental/SkeletalReduction/Source/Private/SkeletalMeshReductionPlugin.cpp b/Engine/Plugins/Experimental/SkeletalReduction/Source/Private/SkeletalMeshReductionPlugin.cpp
--- a/Engine/Plugins/Experimental/SkeletalReduction/Source/Private/SkeletalMeshReductionPlugin.cpp
+++ b/Engine/Plugins/Experimental/SkeletalReduction/Source/Private/SkeletalMeshReductionPlugin.cpp
@@ -22,7 +22,7 @@
 #include "SkeletalSimplifier.h"
 #include "SkeletalMeshReductionSkinnedMesh.h"
 #include "Stats/StatsMisc.h"
-#include "Assets/ClothingAsset.h"
+#include "ClothingAsset.h"
 #include "Factories/FbxSkeletalMeshImportData.h"
 #include "LODUtilities.h"
 
@@ -100,8 +100,7 @@ public:
 	* @returns true if reduction was successful.
 	*/
 	virtual bool ReduceSkeletalMesh( class USkeletalMesh* SkeletalMesh,
-		                             int32 LODIndex,
-		                             bool bReregisterComponent = true
+		                             int32 LODIndex		                             
 	                                ) override ;
 
 
@@ -170,7 +169,8 @@ private:
 		                         FSkeletalMeshOptimizationSettings Settings,
 								 const FImportantBones& ImportantBones,
 		                         const TArray<FMatrix>& BoneMatrices,
-		                         const int32 LODIndex ) const;
+		                         const int32 LODIndex,
+								 const bool bReducingSourceModel) const;
 
 	/**
 	* Remove the specified section from the mesh.
@@ -204,7 +204,8 @@ private:
 		                                 const FSkeletalMeshLODModel& SrcLODModel,
 		                                 const SkeletalSimplifier::FSkinnedSkeletalMesh& SkinnedMesh,
 		                                 const FReferenceSkeleton& RefSkeleton,
-		                                 FSkeletalMeshLODModel& NewModel ) const;
+		                                 FSkeletalMeshLODModel& NewModel,
+										 const bool bReducingSourceModel) const;
 
 	/**
 	* Add the SourceModelInflunces to the LODModel in the case that alternate weights exists.
@@ -359,10 +360,7 @@ void FSkeletalMeshReduction::ShutdownModule()
 
 }
 
-
-
-
-bool FQuadricSkeletalMeshReduction::ReduceSkeletalMesh( USkeletalMesh* SkeletalMesh, int32 LODIndex, bool bReregisterComponent )
+bool FQuadricSkeletalMeshReduction::ReduceSkeletalMesh( USkeletalMesh* SkeletalMesh, int32 LODIndex )
 {
 	check(SkeletalMesh);
 	check(LODIndex >= 0);
@@ -372,21 +370,8 @@ bool FQuadricSkeletalMeshReduction::ReduceSkeletalMesh( USkeletalMesh* SkeletalM
 	check(SkeletalMeshResource);
 	check(LODIndex <= SkeletalMeshResource->LODModels.Num());
 
-	if (bReregisterComponent)
-	{
-		TComponentReregisterContext<USkinnedMeshComponent> ReregisterContext;
-		SkeletalMesh->ReleaseResources();
-		SkeletalMesh->ReleaseResourcesFence.Wait();
-
-		ReduceSkeletalMesh(*SkeletalMesh, *SkeletalMeshResource, LODIndex);
-
-		SkeletalMesh->PostEditChange();
-		SkeletalMesh->InitResources();
-	}
-	else
-	{
-		ReduceSkeletalMesh(*SkeletalMesh, *SkeletalMeshResource, LODIndex);
-	}
+	FScopedSkeletalMeshPostEditChange ScopedPostEditChange(SkeletalMesh);
+	ReduceSkeletalMesh(*SkeletalMesh, *SkeletalMeshResource, LODIndex);
 
 	return true;
 }
@@ -1541,8 +1526,8 @@ void FQuadricSkeletalMeshReduction::ConvertToFSkeletalMeshLODModel( const int32
 	                                                                const FSkeletalMeshLODModel& SrcLODModel,
 	                                                                const SkeletalSimplifier::FSkinnedSkeletalMesh& SkinnedMesh,
 	                                                                const FReferenceSkeleton& RefSkeleton,
-	                                                                FSkeletalMeshLODModel& NewModel
-                                                                   ) const
+	                                                                FSkeletalMeshLODModel& NewModel,	
+																	const bool bReducingSourceModel) const
 {
 	// We might be re-using this model - so clear it.
 
@@ -1553,11 +1538,14 @@ void FQuadricSkeletalMeshReduction::ConvertToFSkeletalMeshLODModel( const int32
 	FSkeletalMeshData SkeletalMeshData;
 	ExtractFSkeletalData(SkinnedMesh, SkeletalMeshData);
 
-	// Add alternate weight data to the NewModel.  Note, this has to be done before we "BuildSkeletalMesh" to insure
-	// that the bone-based vertex chunking respects the alternate weights.
-	// NB: this only prepares the NewModel, but BuildSkeletalMesh is only 1/2-aware of this, so we will have to do some additional work after.
-	
-	AddSourceModelInfluences(SrcLODModel, SkinnedMesh, NewModel);
+	//if (!bReducingSourceModel)
+	{
+		// Add alternate weight data to the NewModel.  Note, this has to be done before we "BuildSkeletalMesh" to insure
+		// that the bone-based vertex chunking respects the alternate weights.
+		// NB: this only prepares the NewModel, but BuildSkeletalMesh is only 1/2-aware of this, so we will have to do some additional work after.
+
+		AddSourceModelInfluences(SrcLODModel, SkinnedMesh, NewModel);
+	}
 	
 
 
@@ -1576,6 +1564,7 @@ void FQuadricSkeletalMeshReduction::ConvertToFSkeletalMeshLODModel( const int32
 	Options.bComputeNormals = false;
 	Options.bComputeTangents = false;
 	Options.bUseMikkTSpace = true; //Avoid builtin build by specifying true for mikkt space
+	Options.bComputeWeightedNormals = false;
 	Options.OverlappingThresholds.ThresholdPosition = 0.0f;
 	Options.OverlappingThresholds.ThresholdTangentNormal = 0.0f;
 	Options.OverlappingThresholds.ThresholdUV = 0.0f;
@@ -1595,15 +1584,19 @@ void FQuadricSkeletalMeshReduction::ConvertToFSkeletalMeshLODModel( const int32
 		Options
 	);
 
+	//Re-Apply the user section changes, the UserSectionsData is map to original section and should match the builded LODModel
+	NewModel.SyncronizeUserSectionsDataArray();
+
 	// Set texture coordinate count on the new model.
 	NewModel.NumTexCoords = SkeletalMeshData.TexCoordCount;
 
 
+	//if (!bReducingSourceModel)
+	{
+		// Update the alternate weights
 
-	// Update the alternate weights
-	
-	UpdateAlternateWeights(MaxBonesPerVertex, NewModel);
-
+		UpdateAlternateWeights(MaxBonesPerVertex, NewModel);
+	}
 }
 
 bool FQuadricSkeletalMeshReduction::ReduceSkeletalLODModel( const FSkeletalMeshLODModel& SrcModel,
@@ -1613,7 +1606,8 @@ bool FQuadricSkeletalMeshReduction::ReduceSkeletalLODModel( const FSkeletalMeshL
 	                                                        FSkeletalMeshOptimizationSettings Settings,
 															const FImportantBones& ImportantBones,
 	                                                        const TArray<FMatrix>& BoneMatrices,
-	                                                        const int32 LODIndex
+	                                                        const int32 LODIndex,
+															const bool bReducingSourceModel
                                                            ) const
 {
 
@@ -1682,7 +1676,7 @@ bool FQuadricSkeletalMeshReduction::ReduceSkeletalLODModel( const FSkeletalMeshL
 
 		// Convert to SkeletalMeshLODModel. 
 
-		ConvertToFSkeletalMeshLODModel(Settings.MaxBonesPerVertex, SrcModel, SkinnedSkeletalMesh, RefSkeleton, OutSkeletalMeshLODModel);
+		ConvertToFSkeletalMeshLODModel(Settings.MaxBonesPerVertex, SrcModel, SkinnedSkeletalMesh, RefSkeleton, OutSkeletalMeshLODModel, bReducingSourceModel);
 
 		// We may need to do additional simplification if the user specified a hard number limit for verts and
 		// the internal chunking during conversion split some verts.
@@ -1804,13 +1798,26 @@ void FQuadricSkeletalMeshReduction::ReduceSkeletalMesh(USkeletalMesh& SkeletalMe
 	struct FSectionData
 	{
 		uint16 MaterialIndex;
+		int32 MaterialMap;
 		bool bCastShadow;
 		bool bRecomputeTangent;
+		bool bDisabled;
+		int32 GenerateUpToLodIndex;
+		int32 ChunkedParentSectionIndex;
+		int32 OriginalDataSectionIndex;
 	};
 
 	TMap<int32, FSectionData> BackupSectionIndexToSectionData;
+	TMap<int32, FSkelMeshSourceSectionUserData> BackupUserSectionsData;
+	FString BackupLodModelBuildStringID = TEXT("");
 	//Store the section data. Store the source LOD model in case we add a LOD, or the target LOD model in case the LOD already exist
 	FSkeletalMeshLODModel& BackupSectionLODModel = bLODModelAdded ? SkeletalMeshResource.LODModels[Settings.BaseLOD] : SkeletalMeshResource.LODModels[LODIndex];
+	if (!bLODModelAdded)
+	{
+		BackupLodModelBuildStringID = BackupSectionLODModel.BuildStringID;
+		BackupUserSectionsData = BackupSectionLODModel.UserSectionsData;
+	}
+
 	BackupSectionIndexToSectionData.Reserve(BackupSectionLODModel.Sections.Num());
 
 	for (int32 SectionIndex = 0; SectionIndex < BackupSectionLODModel.Sections.Num(); ++SectionIndex)
@@ -1825,23 +1832,27 @@ void FQuadricSkeletalMeshReduction::ReduceSkeletalMesh(USkeletalMesh& SkeletalMe
 		SectionData.MaterialIndex = BackupSectionLODModel.Sections[SectionIndex].MaterialIndex;
 		SectionData.bCastShadow = BackupSectionLODModel.Sections[SectionIndex].bCastShadow;
 		SectionData.bRecomputeTangent = BackupSectionLODModel.Sections[SectionIndex].bRecomputeTangent;
+		SectionData.bDisabled = BackupSectionLODModel.Sections[SectionIndex].bDisabled;
+		SectionData.GenerateUpToLodIndex = BackupSectionLODModel.Sections[SectionIndex].GenerateUpToLodIndex;
+		SectionData.ChunkedParentSectionIndex = BackupSectionLODModel.Sections[SectionIndex].ChunkedParentSectionIndex;
+		SectionData.OriginalDataSectionIndex = BackupSectionLODModel.Sections[SectionIndex].OriginalDataSectionIndex;
+		SectionData.MaterialMap = LODInfo->LODMaterialMap.IsValidIndex(SectionIndex) ? LODInfo->LODMaterialMap[SectionIndex] : INDEX_NONE;
+		if (SectionData.MaterialMap == SectionData.MaterialIndex)
+		{
+			//Remove any override if the value is the same
+			SectionData.MaterialMap = INDEX_NONE;
+		}
 	}
 
-	auto FillClothingData = [&SkeletalMeshResource, &LODIndex, bLODModelAdded](int32 &EnableSectionNumber, TArray<bool> &SectionStatus)
+	auto FillSectionMaterialSlot = [&SkeletalMeshResource, &LODIndex, bLODModelAdded](TArray<int32>& SectionMaterialSlot)
 	{
-		EnableSectionNumber = 0;
-		SectionStatus.Empty();
+		SectionMaterialSlot.Empty();
 		if (!bLODModelAdded && SkeletalMeshResource.LODModels.IsValidIndex(LODIndex))
 		{
 			int32 SectionNumber = SkeletalMeshResource.LODModels[LODIndex].Sections.Num();
-			SectionStatus.Reserve(SectionNumber);
 			for (int32 SectionIndex = 0; SectionIndex < SectionNumber; ++SectionIndex)
 			{
-				SectionStatus.Add(!SkeletalMeshResource.LODModels[LODIndex].Sections[SectionIndex].bDisabled);
-				if (SectionStatus[SectionIndex])
-				{
-					EnableSectionNumber++;
-				}
+				SectionMaterialSlot.Add(SkeletalMeshResource.LODModels[LODIndex].Sections[SectionIndex].OriginalDataSectionIndex);
 			}
 		}
 	};
@@ -1849,31 +1860,20 @@ void FQuadricSkeletalMeshReduction::ReduceSkeletalMesh(USkeletalMesh& SkeletalMe
 	// Unbind any existing clothing assets before we reimport the geometry
 	TArray<ClothingAssetUtils::FClothingAssetMeshBinding> ClothingBindings;
 	//Get a map of enable/disable sections
-	int32 OriginalSectionNumberBeforeReduction = 0;
-	TArray<bool> OriginalSectionEnableBeforeReduction;
+	TArray<int32> OriginalSectionMaterialSlot;
 
 	//Do not play with cloth if the LOD is added
 	if (!bLODModelAdded)
 	{
-		//Store the clothBinding
-		ClothingAssetUtils::GetMeshClothingAssetBindings(&SkeletalMesh, ClothingBindings, LODIndex);
-		FillClothingData(OriginalSectionNumberBeforeReduction, OriginalSectionEnableBeforeReduction);
-		//Unbind the Cloth for this LOD before we reduce it, we will put back the cloth after the reduction, if it still match the sections
-		for (ClothingAssetUtils::FClothingAssetMeshBinding& Binding : ClothingBindings)
-		{
-			if (Binding.LODIndex == LODIndex)
-			{
-				Binding.Asset->UnbindFromSkeletalMesh(&SkeletalMesh, Binding.LODIndex);
-			}
-		}
+		FLODUtilities::UnbindClothingAndBackup(&SkeletalMesh, ClothingBindings, LODIndex);
 	}
 
 	bool bReducingSourceModel = false;
-	//Reducing Base LOD, we need to use the temporary data so it can be iterative
-	if (BaseLOD == LODIndex && SkelResource->OriginalReductionSourceMeshData.IsValidIndex(BaseLOD) && !SkelResource->OriginalReductionSourceMeshData[BaseLOD]->IsEmpty())
+	//Reducing source LOD, we need to use the temporary data so it can be iterative
+	if (BaseLOD == LODIndex && (SkelResource->OriginalReductionSourceMeshData.IsValidIndex(BaseLOD) && !SkelResource->OriginalReductionSourceMeshData[BaseLOD]->IsEmpty()))
 	{
 		TMap<FString, TArray<FMorphTargetDelta>> TempLODMorphTargetData;
-		SkelResource->OriginalReductionSourceMeshData[BaseLOD]->LoadReductionData(*SrcModel, TempLODMorphTargetData);
+		SkelResource->OriginalReductionSourceMeshData[BaseLOD]->LoadReductionData(*SrcModel, TempLODMorphTargetData, &SkeletalMesh);
 		bReducingSourceModel = true;
 	}
 	else
@@ -1955,6 +1955,8 @@ void FQuadricSkeletalMeshReduction::ReduceSkeletalMesh(USkeletalMesh& SkeletalMe
 
 	// Swap out the old model.  
 	FSkeletalMeshImportData RawMesh;
+	ESkeletalMeshGeoImportVersions GeoImportVersion = ESkeletalMeshGeoImportVersions::Before_Versionning;
+	ESkeletalMeshSkinningImportVersions SkinningImportVersion = ESkeletalMeshSkinningImportVersions::Before_Versionning;
 	{
 		FSkeletalMeshLODModel* Old = LODModels[LODIndex];
 
@@ -1962,12 +1964,21 @@ void FQuadricSkeletalMeshReduction::ReduceSkeletalMesh(USkeletalMesh& SkeletalMe
 
 		if (!bReducingSourceModel && Old)
 		{
+			//We need to backup the original RawSkeletalMeshBulkData in case it was an imported LOD
+			if (!bLODModelAdded && !Old->RawSkeletalMeshBulkData.IsEmpty())
+			{
+				Old->RawSkeletalMeshBulkData.LoadRawMesh(RawMesh);
+				GeoImportVersion = Old->RawSkeletalMeshBulkData.GeoImportVersion;
+				SkinningImportVersion = Old->RawSkeletalMeshBulkData.SkinningImportVersion;
+			}
 			delete Old;
 		}
 		else if(bReducingSourceModel)
 		{
 			//In case we reduce the source model we want to keep the original import data
 			SrcModel->RawSkeletalMeshBulkData.LoadRawMesh(RawMesh);
+			GeoImportVersion = SrcModel->RawSkeletalMeshBulkData.GeoImportVersion;
+			SkinningImportVersion = SrcModel->RawSkeletalMeshBulkData.SkinningImportVersion;
 		}
 	}
 
@@ -1975,8 +1986,10 @@ void FQuadricSkeletalMeshReduction::ReduceSkeletalMesh(USkeletalMesh& SkeletalMe
 
 	// Reduce LOD model with SrcMesh
 
-	if (ReduceSkeletalLODModel(*SrcModel, *NewModel, SkeletalMesh.GetImportedBounds(), SkeletalMesh.RefSkeleton, Settings, ImportantBones, RelativeToRefPoseMatrices, LODIndex))
+	if (ReduceSkeletalLODModel(*SrcModel, *NewModel, SkeletalMesh.GetImportedBounds(), SkeletalMesh.RefSkeleton, Settings, ImportantBones, RelativeToRefPoseMatrices, LODIndex, bReducingSourceModel))
 	{
+		FSkeletalMeshLODInfo* ReducedLODInfoPtr = SkeletalMesh.GetLODInfo(LODIndex);
+		check(ReducedLODInfoPtr);
 		// Do any joint-welding / bone removal.
 
 		if (MeshBoneReductionInterface != NULL && MeshBoneReductionInterface->GetBoneReductionData(&SkeletalMesh, LODIndex, BonesToRemove))
@@ -1990,35 +2003,22 @@ void FQuadricSkeletalMeshReduction::ReduceSkeletalMesh(USkeletalMesh& SkeletalMe
 
 		if (bOldLodWasFromFile)
 		{
-			SkeletalMesh.GetLODInfo(LODIndex)->LODMaterialMap.Empty();
+			ReducedLODInfoPtr->LODMaterialMap.Empty();
 		}
-
-		// If base lod has a customized LODMaterialMap and this LOD doesn't (could have if changes are applied instead of freshly generated, copy over the data into new new LOD
-
-		if (SkeletalMesh.GetLODInfo(LODIndex)->LODMaterialMap.Num() == 0 && SkeletalMesh.GetLODInfo(BaseLOD)->LODMaterialMap.Num() != 0)
-		{
-			SkeletalMesh.GetLODInfo(LODIndex)->LODMaterialMap = SkeletalMesh.GetLODInfo(BaseLOD)->LODMaterialMap;
-		}
-		else
-		{
-			// Assuming the reducing step has set all material indices correctly, we double check if something went wrong
-			// make sure we don't have more materials
-			int32 TotalSectionCount = NewModel->Sections.Num();
-			if (SkeletalMesh.GetLODInfo(LODIndex)->LODMaterialMap.Num() > TotalSectionCount)
-			{
-				SkeletalMesh.GetLODInfo(LODIndex)->LODMaterialMap = SkeletalMesh.GetLODInfo(BaseLOD)->LODMaterialMap;
-				// Something went wrong during the reduce step during regenerate 					
-				check(SkeletalMesh.GetLODInfo(BaseLOD)->LODMaterialMap.Num() == TotalSectionCount || SkeletalMesh.GetLODInfo(BaseLOD)->LODMaterialMap.Num() == 0);
-			}
-		}
-
 		// Flag this LOD as having been simplified.
-		SkeletalMesh.GetLODInfo(LODIndex)->bHasBeenSimplified = true;
+		ReducedLODInfoPtr->bHasBeenSimplified = true;
 		SkeletalMesh.bHasBeenSimplified = true;
 		//Restore section data
 		FSkeletalMeshLODModel& ImportedModelLOD = SkeletalMesh.GetImportedModel()->LODModels[LODIndex];
 		TArray<bool> SectionMatched;
 		SectionMatched.AddZeroed(ImportedModelLOD.Sections.Num());
+		
+		if (!bLODModelAdded)
+		{
+			ImportedModelLOD.UserSectionsData = BackupUserSectionsData;
+			ImportedModelLOD.BuildStringID = BackupLodModelBuildStringID;
+		}
+
 		for (auto Kvp : BackupSectionIndexToSectionData)
 		{
 			const int32 SourceSectionIndex = Kvp.Key;
@@ -2031,10 +2031,24 @@ void FQuadricSkeletalMeshReduction::ReduceSkeletalMesh(USkeletalMesh& SkeletalMe
 				{
 					continue;
 				}
-				if (SectionData.MaterialIndex == ImportedModelLOD.Sections[SectionIndex].MaterialIndex)
+				if (SectionData.OriginalDataSectionIndex == ImportedModelLOD.Sections[SectionIndex].OriginalDataSectionIndex)
 				{
+					//Restore the LOD material map for this section
+					if (SectionData.MaterialMap != INDEX_NONE)
+					{
+						while (ReducedLODInfoPtr->LODMaterialMap.Num() <= SectionIndex)
+						{
+							ReducedLODInfoPtr->LODMaterialMap.Add(INDEX_NONE);
+						}
+						ReducedLODInfoPtr->LODMaterialMap[SectionIndex] = SectionData.MaterialMap;
+					}
+					//Restore the LODModel Section
 					ImportedModelLOD.Sections[SectionIndex].bCastShadow = SectionData.bCastShadow;
 					ImportedModelLOD.Sections[SectionIndex].bRecomputeTangent = SectionData.bRecomputeTangent;
+					ImportedModelLOD.Sections[SectionIndex].bDisabled = SectionData.bDisabled;
+					ImportedModelLOD.Sections[SectionIndex].GenerateUpToLodIndex = SectionData.GenerateUpToLodIndex;
+					ImportedModelLOD.Sections[SectionIndex].ChunkedParentSectionIndex = SectionData.ChunkedParentSectionIndex;
+					ImportedModelLOD.Sections[SectionIndex].OriginalDataSectionIndex = SectionData.OriginalDataSectionIndex;
 					SectionMatched[SectionIndex] = true;
 					break;
 				}
@@ -2087,47 +2101,22 @@ void FQuadricSkeletalMeshReduction::ReduceSkeletalMesh(USkeletalMesh& SkeletalMe
 	if (!bLODModelAdded)
 	{
 		//Get the number of enabled section
-		int32 SectionNumberAfterReduction = 0;
-		TArray<bool> SectionEnableAfterReduction;
-		FillClothingData(SectionNumberAfterReduction, SectionEnableAfterReduction);
+		TArray<int32> SectionMaterialSlotAfterReduction;
+		FillSectionMaterialSlot(SectionMaterialSlotAfterReduction);
 
-		//Put back the clothing for this newly reduce LOD only if the section count match.
-		if (ClothingBindings.Num() > 0 && OriginalSectionNumberBeforeReduction == SectionNumberAfterReduction)
+		//Put back the clothing for this newly reduce LOD
+		if (ClothingBindings.Num() > 0)
 		{
-			TArray<int32> RemapSectionIndex;
-			int32 SectionIndexTest = 0;
-			for (int32 SectionIndexRef = 0; SectionIndexRef < OriginalSectionEnableBeforeReduction.Num(); SectionIndexRef++)
-			{
-				int32& RemapValue = RemapSectionIndex.Add_GetRef(INDEX_NONE);
-				if (!OriginalSectionEnableBeforeReduction[SectionIndexRef])
-				{
-					continue;
-				}
-				for (; SectionIndexTest <= SectionIndexRef; SectionIndexTest++)
-				{
-					if (SectionEnableAfterReduction.IsValidIndex(SectionIndexTest) && SectionEnableAfterReduction[SectionIndexTest])
-					{
-						RemapValue = SectionIndexTest++;
-						break;
-					}
-				}
-			}
-
-			for (ClothingAssetUtils::FClothingAssetMeshBinding& Binding : ClothingBindings)
-			{
-				int32 RemapBindingSectionIndex = RemapSectionIndex[Binding.SectionIndex];
-				if (RemapBindingSectionIndex != INDEX_NONE && Binding.LODIndex == LODIndex && NewModel->Sections.IsValidIndex(RemapBindingSectionIndex))
-				{
-					Binding.Asset->BindToSkeletalMesh(&SkeletalMesh, Binding.LODIndex, RemapBindingSectionIndex, Binding.AssetInternalLodIndex, false);
-				}
-			}
+			FLODUtilities::RestoreClothingFromBackup(&SkeletalMesh, ClothingBindings, LODIndex);
 		}
 	}
 
-	if (bReducingSourceModel && RawMesh.Points.Num() > 0)
+	if ((bReducingSourceModel || !bLODModelAdded ) && RawMesh.Points.Num() > 0)
 	{
 		//Put back the original import data, we need it to allow inline reduction and skeletal mesh split workflow
 		SkeletalMeshResource.LODModels[LODIndex].RawSkeletalMeshBulkData.SaveRawMesh(RawMesh);
+		SkeletalMeshResource.LODModels[LODIndex].RawSkeletalMeshBulkData.GeoImportVersion = GeoImportVersion;
+		SkeletalMeshResource.LODModels[LODIndex].RawSkeletalMeshBulkData.SkinningImportVersion = SkinningImportVersion;
 	}
 
 	SkeletalMesh.CalculateRequiredBones(SkeletalMeshResource.LODModels[LODIndex], SkeletalMesh.RefSkeleton, &BonesToRemove);
diff --git a/Engine/Plugins/Experimental/SkeletalReduction/Source/SkeletalMeshReduction.Build.cs b/Engine/Plugins/Experimental/SkeletalReduction/Source/SkeletalMeshReduction.Build.cs
--- a/Engine/Plugins/Experimental/SkeletalReduction/Source/SkeletalMeshReduction.Build.cs
+++ b/Engine/Plugins/Experimental/SkeletalReduction/Source/SkeletalMeshReduction.Build.cs
@@ -41,7 +41,8 @@ namespace UnrealBuildTool.Rules
                     "EditorStyle",
                     "RenderCore",
                     "RHI",
-                    "QuadricMeshReduction"
+                    "QuadricMeshReduction",
+					"SkeletalMeshUtilitiesCommon"
 					// ... add other public dependencies that you statically link with here ...
 				}
                 );
@@ -56,8 +57,7 @@ namespace UnrealBuildTool.Rules
                     "AnimationModifiers",
                     "MeshBoneReduction",
                     "QuadricMeshReduction",
-                    "ClothingSystemRuntimeNv",
-                    "SkeletalMeshUtilitiesCommon",
+                    "ClothingSystemRuntimeCommon"
 					// ... add private dependencies that you statically link with here ...
                     // QuadricMeshReduction is only for testing
 				}
diff --git a/Engine/Plugins/Runtime/ApexDestruction/Source/ApexDestruction/Public/DestructibleComponent.h b/Engine/Plugins/Runtime/ApexDestruction/Source/ApexDestruction/Public/DestructibleComponent.h
--- a/Engine/Plugins/Runtime/ApexDestruction/Source/ApexDestruction/Public/DestructibleComponent.h
+++ b/Engine/Plugins/Runtime/ApexDestruction/Source/ApexDestruction/Public/DestructibleComponent.h
@@ -7,7 +7,6 @@
 #include "EngineDefines.h"
 #include "PhysxUserData.h"
 #include "ApexDestructionCustomPayload.h"
-#include "PhysicsEngine/BodyInstance.h"
 #include "Components/SkinnedMeshComponent.h"
 #include "DestructibleInterface.h"
 #include "DestructibleComponent.generated.h"
@@ -15,6 +14,7 @@
 class AController;
 class UDestructibleComponent;
 class USkeletalMesh;
+struct FBodyInstance;
 struct FCollisionShape;
 struct FNavigableGeometryExport;
 
diff --git a/Engine/Source/Developer/MeshBuilder/MeshBuilder.Build.cs b/Engine/Source/Developer/MeshBuilder/MeshBuilder.Build.cs
--- a/Engine/Source/Developer/MeshBuilder/MeshBuilder.Build.cs
+++ b/Engine/Source/Developer/MeshBuilder/MeshBuilder.Build.cs
@@ -16,17 +16,19 @@ namespace UnrealBuildTool.Rules
                     "Engine",
                     "RenderCore",
                     "MeshDescription",
-                    "RenderCore",
+					"StaticMeshDescription",
                     "MeshDescriptionOperations",
                     "MeshReductionInterface",
                     "RawMesh",
+					"MeshUtilities",
                     "MeshUtilitiesCommon",
+                    "ClothingSystemRuntimeNv",
+                    "MeshBoneReduction",
+                    "SkeletalMeshUtilitiesCommon",
+					"MeshBuilderCommon",
                 }
 			);
 
-			AddEngineThirdPartyPrivateStaticDependencies(Target, "nvTriStrip");
-            AddEngineThirdPartyPrivateStaticDependencies(Target, "ForsythTriOptimizer");
-	        AddEngineThirdPartyPrivateStaticDependencies(Target, "nvTessLib");
             AddEngineThirdPartyPrivateStaticDependencies(Target, "QuadricMeshReduction");
        }
 	}
diff --git a/Engine/Source/Developer/MeshBuilder/Private/SkeletalMeshBuilder.cpp b/Engine/Source/Developer/MeshBuilder/Private/SkeletalMeshBuilder.cpp
new file mode 100644
--- /dev/null
+++ b/Engine/Source/Developer/MeshBuilder/Private/SkeletalMeshBuilder.cpp
@@ -0,0 +1,197 @@
+// Copyright 1998-2019 Epic Games, Inc. All Rights Reserved.
+
+#include "SkeletalMeshBuilder.h"
+#include "Modules/ModuleManager.h"
+#include "MeshBoneReduction.h"
+#include "Engine/EngineTypes.h"
+#include "Engine/SkeletalMesh.h"
+#include "PhysicsEngine/BodySetup.h"
+#include "MeshDescription.h"
+#include "MeshDescriptionOperations.h"
+#include "MeshAttributes.h"
+#include "MeshDescriptionHelper.h"
+#include "MeshBuild.h"
+#include "Rendering/SkeletalMeshModel.h"
+#include "Rendering/SkeletalMeshLODModel.h"
+#include "GPUSkinVertexFactory.h"
+#include "ThirdPartyBuildOptimizationHelper.h"
+#include "Misc/ScopedSlowTask.h"
+#include "LODUtilities.h"
+#include "ClothingAsset.h"
+#include "MeshUtilities.h"
+#include "EditorFramework/AssetImportData.h"
+
+DEFINE_LOG_CATEGORY(LogSkeletalMeshBuilder);
+
+namespace SkeletalMeshBuilderOptimization
+{
+	void CacheOptimizeIndexBuffer(TArray<uint16>& Indices)
+	{
+		BuildOptimizationThirdParty::CacheOptimizeIndexBuffer(Indices);
+	}
+
+	void CacheOptimizeIndexBuffer(TArray<uint32>& Indices)
+	{
+		BuildOptimizationThirdParty::CacheOptimizeIndexBuffer(Indices);
+	}
+}
+
+struct InfluenceMap
+{
+	FORCEINLINE bool operator()(const float& A, const float& B) const
+	{
+		return B > A;
+	}
+};
+
+struct FSkeletalMeshVertInstanceIDAndZ
+{
+	FVertexInstanceID Index;
+	float Z;
+};
+
+FSkeletalMeshBuilder::FSkeletalMeshBuilder()
+{
+
+}
+
+
+bool FSkeletalMeshBuilder::Build(USkeletalMesh* SkeletalMesh, const int32 LODIndex, const bool bRegenDepLODs)
+{
+	check(SkeletalMesh->GetImportedModel());
+	check(SkeletalMesh->GetImportedModel()->LODModels.IsValidIndex(LODIndex));
+	check(SkeletalMesh->GetLODInfo(LODIndex) != nullptr);
+	
+	FSkeletalMeshLODInfo* LODInfo = SkeletalMesh->GetLODInfo(LODIndex);
+	//We want to backup in case the LODModel is regenerated, this data is use to validate in the UI if the ddc must be rebuild
+	const FString BackupBuildStringID = SkeletalMesh->GetImportedModel()->LODModels[LODIndex].BuildStringID;
+
+	const FReferenceSkeleton& RefSkeleton = SkeletalMesh->RefSkeleton;
+
+	FScopedSlowTask SlowTask(6.0f, NSLOCTEXT("SkeltalMeshBuilder", "BuildingSkeletalMeshLOD", "Building skeletal mesh LOD"));
+	SlowTask.MakeDialog();
+
+	//Prevent any PostEdit change during the build
+	FScopedSkeletalMeshPostEditChange ScopedPostEditChange(SkeletalMesh, false, false);
+	// Unbind any existing clothing assets before we reimport the geometry
+	TArray<ClothingAssetUtils::FClothingAssetMeshBinding> ClothingBindings;
+	FLODUtilities::UnbindClothingAndBackup(SkeletalMesh, ClothingBindings, LODIndex);
+
+	FSkeletalMeshImportData SkeletalMeshImportData;
+	float ProgressStepSize = SkeletalMesh->IsReductionActive(LODIndex) ? 0.5f : 1.0f;
+
+	//This scope define where we can use the LODModel, after a reduction the LODModel must be requery since it is a new instance
+	{
+		FSkeletalMeshLODModel& BuildLODModel = SkeletalMesh->GetImportedModel()->LODModels[LODIndex];
+		
+		BuildLODModel.RawSkeletalMeshBulkData.LoadRawMesh(SkeletalMeshImportData);
+
+		TArray<FVector> LODPoints;
+		TArray<SkeletalMeshImportData::FMeshWedge> LODWedges;
+		TArray<SkeletalMeshImportData::FMeshFace> LODFaces;
+		TArray<SkeletalMeshImportData::FVertInfluence> LODInfluences;
+		TArray<int32> LODPointToRawMap;
+		SkeletalMeshImportData.CopyLODImportData(LODPoints, LODWedges, LODFaces, LODInfluences, LODPointToRawMap);
+
+
+		//Build the skeletalmesh using mesh utilities module
+		IMeshUtilities::MeshBuildOptions Options;
+		Options.FillOptions(LODInfo->BuildSettings);
+		IMeshUtilities& MeshUtilities = FModuleManager::Get().LoadModuleChecked<IMeshUtilities>("MeshUtilities");
+
+		// Create skinning streams for NewModel.
+		SlowTask.EnterProgressFrame(1.0f);
+		MeshUtilities.BuildSkeletalMesh(
+			BuildLODModel,
+			RefSkeleton,
+			LODInfluences,
+			LODWedges,
+			LODFaces,
+			LODPoints,
+			LODPointToRawMap,
+			Options
+		);
+
+		//Re-Apply the user section changes, the UserSectionsData is map to original section and should match the builded LODModel
+		BuildLODModel.SyncronizeUserSectionsDataArray();
+
+		// Set texture coordinate count on the new model.
+		BuildLODModel.NumTexCoords = SkeletalMeshImportData.NumTexCoords;
+
+		//Re-apply the morph target
+		SlowTask.EnterProgressFrame(1.0f, NSLOCTEXT("SkeltalMeshBuilder", "RebuildMorphTarget", "Rebuilding morph targets..."));
+		if (SkeletalMeshImportData.MorphTargetNames.Num() > 0)
+		{
+			FLODUtilities::BuildMorphTargets(SkeletalMesh, SkeletalMeshImportData, LODIndex, !Options.bComputeNormals, !Options.bComputeTangents, Options.bUseMikkTSpace);
+		}
+
+		//Re-apply the alternate skinning it must be after the inline reduction
+		SlowTask.EnterProgressFrame(1.0f, NSLOCTEXT("SkeltalMeshBuilder", "RebuildAlternateSkinning", "Rebuilding alternate skinning..."));
+		const TArray<FSkinWeightProfileInfo>& SkinProfiles = SkeletalMesh->GetSkinWeightProfiles();
+		for (int32 SkinProfileIndex = 0; SkinProfileIndex < SkinProfiles.Num(); ++SkinProfileIndex)
+		{
+			const FSkinWeightProfileInfo& ProfileInfo = SkinProfiles[SkinProfileIndex];
+			FLODUtilities::UpdateAlternateSkinWeights(SkeletalMesh, ProfileInfo.Name, LODIndex, Options.OverlappingThresholds, !Options.bComputeNormals, !Options.bComputeTangents, Options.bUseMikkTSpace, Options.bComputeWeightedNormals);
+		}
+
+		
+		FSkeletalMeshUpdateContext UpdateContext;
+		UpdateContext.SkeletalMesh = SkeletalMesh;
+		//We are reduce ourself in this case we reduce ourself from the original data and return true.
+		if (SkeletalMesh->IsReductionActive(LODIndex))
+		{
+			SlowTask.EnterProgressFrame(ProgressStepSize, NSLOCTEXT("SkeltalMeshBuilder", "RegenerateLOD", "Regenerate LOD..."));
+			//Update the original reduction data since we just build a new LODModel.
+			if (LODInfo->ReductionSettings.BaseLOD == LODIndex && SkeletalMesh->GetImportedModel()->OriginalReductionSourceMeshData.IsValidIndex(LODIndex))
+			{
+				//Make the copy of the data only once until the ImportedModel change (re-imported)
+				SkeletalMesh->GetImportedModel()->OriginalReductionSourceMeshData[LODIndex]->EmptyBulkData();
+				TMap<FString, TArray<FMorphTargetDelta>> BaseLODMorphTargetData;
+				BaseLODMorphTargetData.Empty(SkeletalMesh->MorphTargets.Num());
+				for (UMorphTarget *MorphTarget : SkeletalMesh->MorphTargets)
+				{
+					if (!MorphTarget->HasDataForLOD(LODIndex))
+					{
+						continue;
+					}
+					TArray<FMorphTargetDelta>& MorphDeltasArray = BaseLODMorphTargetData.FindOrAdd(MorphTarget->GetFullName());
+					const FMorphTargetLODModel& BaseMorphModel = MorphTarget->MorphLODModels[LODIndex];
+					//Iterate each original morph target source index to fill the NewMorphTargetDeltas array with the TargetMatchData.
+					for (const FMorphTargetDelta& MorphDelta : BaseMorphModel.Vertices)
+					{
+						MorphDeltasArray.Add(MorphDelta);
+					}
+				}
+				//Copy the original SkeletalMesh LODModel
+				SkeletalMesh->GetImportedModel()->OriginalReductionSourceMeshData[LODIndex]->SaveReductionData(BuildLODModel, BaseLODMorphTargetData, SkeletalMesh);
+
+				if (LODIndex == 0)
+				{
+					SkeletalMesh->GetLODInfo(LODIndex)->SourceImportFilename = SkeletalMesh->AssetImportData->GetFirstFilename();
+				}
+			}
+			FLODUtilities::SimplifySkeletalMeshLOD(UpdateContext, LODIndex, false);
+		}
+	}
+
+	FSkeletalMeshLODModel& LODModelAfterReduction = SkeletalMesh->GetImportedModel()->LODModels[LODIndex];
+	
+	//Re-apply the clothing using the UserSectionsData, this will ensure we remap correctly the cloth if the reduction has change the number of sections
+	SlowTask.EnterProgressFrame(1.0f, NSLOCTEXT("SkeltalMeshBuilder", "RebuildClothing", "Rebuilding clothing..."));
+	FLODUtilities::RestoreClothingFromBackup(SkeletalMesh, ClothingBindings, LODIndex);
+
+	LODModelAfterReduction.SyncronizeUserSectionsDataArray();
+	LODModelAfterReduction.NumTexCoords = SkeletalMeshImportData.NumTexCoords;
+	LODModelAfterReduction.BuildStringID = BackupBuildStringID;
+
+	SlowTask.EnterProgressFrame(ProgressStepSize, NSLOCTEXT("SkeltalMeshBuilder", "RegenerateDependentLODs", "Regenerate Dependent LODs..."));
+	if (bRegenDepLODs)
+	{
+		//Regenerate dependent LODs
+		FLODUtilities::RegenerateDependentLODs(SkeletalMesh, LODIndex);
+	}
+
+	SlowTask.EnterProgressFrame();
+
+	return true;
+}
diff --git a/Engine/Source/Developer/SkeletalMeshUtilitiesCommon/Private/LODUtilities.cpp b/Engine/Source/Developer/SkeletalMeshUtilitiesCommon/Private/LODUtilities.cpp
--- a/Engine/Source/Developer/SkeletalMeshUtilitiesCommon/Private/LODUtilities.cpp
+++ b/Engine/Source/Developer/SkeletalMeshUtilitiesCommon/Private/LODUtilities.cpp
@@ -14,8 +14,8 @@
 #include "Engine/SkeletalMesh.h"
 #include "EditorFramework/AssetImportData.h"
 #include "MeshUtilities.h"
+#include "ClothingAsset.h"
 #include "OverlappingCorners.h"
-#include "Assets/ClothingAsset.h"
 #include "Framework/Commands/UIAction.h"
 
 #include "ObjectTools.h"
@@ -202,14 +202,11 @@ bool FLODUtilities::RegenerateLOD(USkeletalMesh* SkeletalMesh, int32 NewLODCount
 {
 	if (SkeletalMesh)
 	{
+		FScopedSkeletalMeshPostEditChange ScopedPostEditChange(SkeletalMesh);
+
 		// Unbind any existing clothing assets before we regenerate all LODs
 		TArray<ClothingAssetUtils::FClothingAssetMeshBinding> ClothingBindings;
-		ClothingAssetUtils::GetMeshClothingAssetBindings(SkeletalMesh, ClothingBindings);
-
-		for (ClothingAssetUtils::FClothingAssetMeshBinding& Binding : ClothingBindings)
-		{
-			Binding.Asset->UnbindFromSkeletalMesh(SkeletalMesh, Binding.LODIndex);
-		}
+		FLODUtilities::UnbindClothingAndBackup(SkeletalMesh, ClothingBindings);
 
 		int32 LODCount = SkeletalMesh->GetLODNum();
 
@@ -239,7 +236,7 @@ bool FLODUtilities::RegenerateLOD(USkeletalMesh* SkeletalMesh, int32 NewLODCount
 			for (int32 LODIdx = CurrentNumLODs; LODIdx < LODCount; LODIdx++)
 			{
 				// if no previous setting found, it will use default setting. 
-				FLODUtilities::SimplifySkeletalMeshLOD(UpdateContext, LODIdx);
+				FLODUtilities::SimplifySkeletalMeshLOD(UpdateContext, LODIdx, false);
 			}
 		}
 		else
@@ -249,22 +246,13 @@ bool FLODUtilities::RegenerateLOD(USkeletalMesh* SkeletalMesh, int32 NewLODCount
 				FSkeletalMeshLODInfo& CurrentLODInfo = *(SkeletalMesh->GetLODInfo(LODIdx));
 				if ((bRegenerateEvenIfImported && LODIdx > 0) || (bGenerateBaseLOD && LODIdx == 0) || CurrentLODInfo.bHasBeenSimplified )
 				{
-					FLODUtilities::SimplifySkeletalMeshLOD(UpdateContext, LODIdx);
+					FLODUtilities::SimplifySkeletalMeshLOD(UpdateContext, LODIdx, false);
 				}
 			}
 		}
 
 		//Restore all clothing we can
-		for (ClothingAssetUtils::FClothingAssetMeshBinding& Binding : ClothingBindings)
-		{
-			if (SkeletalMesh->GetImportedModel()->LODModels.IsValidIndex(Binding.LODIndex) &&
-				SkeletalMesh->GetImportedModel()->LODModels[Binding.LODIndex].Sections.IsValidIndex(Binding.SectionIndex))
-			{
-				Binding.Asset->BindToSkeletalMesh(SkeletalMesh, Binding.LODIndex, Binding.SectionIndex, Binding.AssetInternalLodIndex);
-			}
-		}
-
-		SkeletalMesh->PostEditChange();
+		FLODUtilities::RestoreClothingFromBackup(SkeletalMesh, ClothingBindings);
 
 		return true;
 	}
@@ -296,21 +284,12 @@ void FLODUtilities::RemoveLOD(FSkeletalMeshUpdateContext& UpdateContext, int32 D
 	// If its a valid LOD, kill it.
 	if( DesiredLOD > 0 && DesiredLOD < SkelMeshModel->LODModels.Num() )
 	{
-		//We'll be modifying the skel mesh data so reregister
-
-		//TODO - do we need to reregister something else instead?
-		FMultiComponentReregisterContext ReregisterContext(UpdateContext.AssociatedComponents);
-
-		// Release rendering resources before deleting LOD
-		SkeletalMesh->ReleaseResources();
+		FScopedSkeletalMeshPostEditChange ScopedPostEditChange(SkeletalMesh);
 
 		// Block until this is done
-		FlushRenderingCommands();
 
 		SkelMeshModel->LODModels.RemoveAt(DesiredLOD);
 		SkeletalMesh->RemoveLODInfo(DesiredLOD);
-		SkeletalMesh->InitResources();
-
 		RefreshLODChange(SkeletalMesh);
 
 		// Set the forced LOD to Auto.
@@ -332,9 +311,6 @@ void FLODUtilities::RemoveLOD(FSkeletalMeshUpdateContext& UpdateContext, int32 D
 			}
 		}
 
-		// This will recache derived render data, and re-init resources
-		SkeletalMesh->PostEditChange();
-
 		//Notify calling system of change
 		UpdateContext.OnLODChanged.ExecuteIfBound();
 
@@ -488,7 +464,7 @@ void ProjectTargetOnBase(const TArray<FSoftSkinVertex>& BaseVertices, const TArr
 	{
 		//Use the remap base index in case some sections disappear during the reduce phase
 		int32 BaseSectionIndex = TargetSectionMatchBaseIndex[SectionIndex];
-		if (BaseSectionIndex == INDEX_NONE || !PerSectionBaseTriangleIndices.IsValidIndex(BaseSectionIndex))
+		if (BaseSectionIndex == INDEX_NONE || !PerSectionBaseTriangleIndices.IsValidIndex(BaseSectionIndex) || PerSectionBaseTriangleIndices[BaseSectionIndex].Num() < 1)
 		{
 			continue;
 		}
@@ -518,6 +494,7 @@ void ProjectTargetOnBase(const TArray<FSoftSkinVertex>& BaseVertices, const TArr
 			TriangleElement.TriangleIndex = Triangles.Num();
 			Triangles.Add(TriangleElement);
 		}
+		check(!BaseMeshUVBound.GetExtent().IsNearlyZero());
 		//Setup the Quad tree
 		float UVsQuadTreeMinSize = 0.001f;
 		TQuadTree<uint32, 100> QuadTree(BaseMeshUVBound, UVsQuadTreeMinSize);
@@ -661,7 +638,7 @@ void CreateLODMorphTarget(USkeletalMesh* SkeletalMesh, FReductionBaseSkeletalMes
 	if (ReductionBaseSkeletalMeshBulkData != nullptr)
 	{
 		FSkeletalMeshLODModel TempBaseLODModel;
-		ReductionBaseSkeletalMeshBulkData->LoadReductionData(TempBaseLODModel, BaseLODMorphTargetData);
+		ReductionBaseSkeletalMeshBulkData->LoadReductionData(TempBaseLODModel, BaseLODMorphTargetData, SkeletalMesh);
 	}
 
 	FSkeletalMeshModel* SkeletalMeshModel = SkeletalMesh->GetImportedModel();
@@ -841,7 +818,8 @@ void FLODUtilities::ApplyMorphTargetsToLOD(USkeletalMesh* SkeletalMesh, int32 So
 	TMap<FString, TArray<FMorphTargetDelta>> TempBaseLODMorphTargetData;
 	if (bReduceBaseLOD)
 	{
-		ReductionBaseSkeletalMeshBulkData->LoadReductionData(TempBaseLODModel, TempBaseLODMorphTargetData);
+		check(ReductionBaseSkeletalMeshBulkData != nullptr);
+		ReductionBaseSkeletalMeshBulkData->LoadReductionData(TempBaseLODModel, TempBaseLODMorphTargetData, SkeletalMesh);
 	}
 
 	const FSkeletalMeshLODModel& BaseLODModel = bReduceBaseLOD ? TempBaseLODModel : SkeletalMeshResource->LODModels[SourceLOD];
@@ -900,7 +878,7 @@ void FLODUtilities::ApplyMorphTargetsToLOD(USkeletalMesh* SkeletalMesh, int32 So
 	for (int32 SectionIndex = 0; SectionIndex < SectionCount; ++SectionIndex)
 	{
 		const FSkelMeshSection& Section = BaseLODModel.Sections[SectionIndex];
-		if (Section.HasClothingData())
+		if (Section.ClothingData.AssetGuid.IsValid())
 		{
 			continue;
 		}
@@ -964,13 +942,14 @@ void FLODUtilities::ApplyMorphTargetsToLOD(USkeletalMesh* SkeletalMesh, int32 So
 	CreateLODMorphTarget(SkeletalMesh, ReductionBaseSkeletalMeshBulkData, SourceLOD, DestinationLOD, PerMorphTargetBaseIndexToMorphTargetDelta, BaseMorphIndexToTargetIndexList, TargetVertices, TargetMatchData);
 }
 
-void FLODUtilities::SimplifySkeletalMeshLOD( USkeletalMesh* SkeletalMesh, int32 DesiredLOD, bool bReregisterComponent /*= true*/, bool bRestoreClothing /*= false*/)
+void FLODUtilities::SimplifySkeletalMeshLOD( USkeletalMesh* SkeletalMesh, int32 DesiredLOD, bool bRestoreClothing /*= false*/)
 {
 	IMeshReductionModule& ReductionModule = FModuleManager::Get().LoadModuleChecked<IMeshReductionModule>("MeshReductionInterface");
 	IMeshReduction* MeshReduction = ReductionModule.GetSkeletalMeshReductionInterface();
 
 	check (MeshReduction && MeshReduction->IsSupported());
 
+
 	if (DesiredLOD == 0
 		&& SkeletalMesh->GetLODInfo(DesiredLOD) != nullptr
 		&& SkeletalMesh->GetLODInfo(DesiredLOD)->bHasBeenSimplified
@@ -1000,18 +979,13 @@ void FLODUtilities::SimplifySkeletalMeshLOD( USkeletalMesh* SkeletalMesh, int32
 		GWarn->BeginSlowTask(StatusUpdate, true);
 	}
 
+	FScopedSkeletalMeshPostEditChange ScopedPostEditChange(SkeletalMesh);
+
 	// Unbind DesiredLOD existing clothing assets before we simplify this LOD
 	TArray<ClothingAssetUtils::FClothingAssetMeshBinding> ClothingBindings;
-	if (bRestoreClothing)
+	if (bRestoreClothing && SkeletalMesh->GetImportedModel() && SkeletalMesh->GetImportedModel()->LODModels.IsValidIndex(DesiredLOD))
 	{
-		ClothingAssetUtils::GetMeshClothingAssetBindings(SkeletalMesh, ClothingBindings);
-		for (ClothingAssetUtils::FClothingAssetMeshBinding& Binding : ClothingBindings)
-		{
-			if (DesiredLOD == Binding.LODIndex)
-			{
-				Binding.Asset->UnbindFromSkeletalMesh(SkeletalMesh, Binding.LODIndex);
-			}
-		}
+		FLODUtilities::UnbindClothingAndBackup(SkeletalMesh, ClothingBindings, DesiredLOD);
 	}
 
 	if (SkeletalMesh->GetLODInfo(DesiredLOD) != nullptr)
@@ -1048,8 +1022,20 @@ void FLODUtilities::SimplifySkeletalMeshLOD( USkeletalMesh* SkeletalMesh, int32
 						MorphDeltasArray.Add(MorphDelta);
 					}
 				}
+				
 				//Copy the original SkeletalMesh LODModel
-				SkeletalMeshResource->OriginalReductionSourceMeshData[DesiredLOD]->SaveReductionData(SrcModel, BaseLODMorphTargetData);
+				// Unbind clothing before saving the original data, we must not restore clothing to do inline reduction
+				{
+					TArray<ClothingAssetUtils::FClothingAssetMeshBinding> TemporaryRemoveClothingBindings;
+					FLODUtilities::UnbindClothingAndBackup(SkeletalMesh, TemporaryRemoveClothingBindings, DesiredLOD);
+
+					SkeletalMeshResource->OriginalReductionSourceMeshData[DesiredLOD]->SaveReductionData(SrcModel, BaseLODMorphTargetData, SkeletalMesh);
+
+					if (TemporaryRemoveClothingBindings.Num() > 0)
+					{
+						FLODUtilities::RestoreClothingFromBackup(SkeletalMesh, TemporaryRemoveClothingBindings, DesiredLOD);
+					}
+				}
 
 				if (DesiredLOD == 0)
 				{
@@ -1060,7 +1046,7 @@ void FLODUtilities::SimplifySkeletalMeshLOD( USkeletalMesh* SkeletalMesh, int32
 	}
 	
 
-	if (MeshReduction->ReduceSkeletalMesh(SkeletalMesh, DesiredLOD, bReregisterComponent))
+	if (MeshReduction->ReduceSkeletalMesh(SkeletalMesh, DesiredLOD))
 	{
 		check(SkeletalMesh->GetLODNum() >= 1);
 
@@ -1078,21 +1064,7 @@ void FLODUtilities::SimplifySkeletalMeshLOD( USkeletalMesh* SkeletalMesh, int32
 			}
 		};
 
-		if (bReregisterComponent)
-		{
-			TComponentReregisterContext<USkinnedMeshComponent> ReregisterContext;
-			SkeletalMesh->ReleaseResources();
-			SkeletalMesh->ReleaseResourcesFence.Wait();
-
-			ApplyMorphTargetOption();
-			
-			SkeletalMesh->PostEditChange();
-			SkeletalMesh->InitResources();
-		}
-		else
-		{
-			ApplyMorphTargetOption();
-		}
+		ApplyMorphTargetOption();
 		SkeletalMesh->MarkPackageDirty();
 	}
 	else
@@ -1105,25 +1077,15 @@ void FLODUtilities::SimplifySkeletalMeshLOD( USkeletalMesh* SkeletalMesh, int32
 	}
 
 	//Put back the clothing for the DesiredLOD
-	if (bRestoreClothing)
+	if (bRestoreClothing && ClothingBindings.Num() > 0 && SkeletalMesh->GetImportedModel()->LODModels.IsValidIndex(DesiredLOD))
 	{
-		for (ClothingAssetUtils::FClothingAssetMeshBinding& Binding : ClothingBindings)
-		{
-			if (SkeletalMesh->GetImportedModel()->LODModels.IsValidIndex(Binding.LODIndex) &&
-				SkeletalMesh->GetImportedModel()->LODModels[Binding.LODIndex].Sections.IsValidIndex(Binding.SectionIndex))
-			{
-				if (DesiredLOD == Binding.LODIndex)
-				{
-					Binding.Asset->BindToSkeletalMesh(SkeletalMesh, Binding.LODIndex, Binding.SectionIndex, Binding.AssetInternalLodIndex);
-				}
-			}
-		}
+		FLODUtilities::RestoreClothingFromBackup(SkeletalMesh, ClothingBindings, DesiredLOD);
 	}
 
 	GWarn->EndSlowTask();
 }
 
-void FLODUtilities::SimplifySkeletalMeshLOD(FSkeletalMeshUpdateContext& UpdateContext, int32 DesiredLOD, bool bReregisterComponent /*= true*/, bool bRestoreClothing /*= false*/)
+void FLODUtilities::SimplifySkeletalMeshLOD(FSkeletalMeshUpdateContext& UpdateContext, int32 DesiredLOD, bool bRestoreClothing /*= false*/)
 {
 	USkeletalMesh* SkeletalMesh = UpdateContext.SkeletalMesh;
 	IMeshReductionModule& ReductionModule = FModuleManager::Get().LoadModuleChecked<IMeshReductionModule>("MeshReductionInterface");
@@ -1131,7 +1093,7 @@ void FLODUtilities::SimplifySkeletalMeshLOD(FSkeletalMeshUpdateContext& UpdateCo
 
 	if (MeshReduction && MeshReduction->IsSupported() && SkeletalMesh)
 	{
-		SimplifySkeletalMeshLOD(SkeletalMesh, DesiredLOD, bReregisterComponent, bRestoreClothing);
+		SimplifySkeletalMeshLOD(SkeletalMesh, DesiredLOD, bRestoreClothing);
 		
 		if (UpdateContext.OnLODChanged.IsBound())
 		{
@@ -1141,40 +1103,27 @@ void FLODUtilities::SimplifySkeletalMeshLOD(FSkeletalMeshUpdateContext& UpdateCo
 	}
 }
 
-void FLODUtilities::RestoreSkeletalMeshLODImportedData(USkeletalMesh* SkeletalMesh, int32 LodIndex, bool bReregisterComponent /*= true*/)
+void FLODUtilities::RestoreSkeletalMeshLODImportedData(USkeletalMesh* SkeletalMesh, int32 LodIndex)
 {
-
 	if (!SkeletalMesh->GetImportedModel()->OriginalReductionSourceMeshData.IsValidIndex(LodIndex) || SkeletalMesh->GetImportedModel()->OriginalReductionSourceMeshData[LodIndex]->IsEmpty())
 	{
 		//There is nothing to restore
 		return;
 	}
 
+	FScopedSkeletalMeshPostEditChange ScopedPostEditChange(SkeletalMesh);
+
 	// Unbind LodIndex existing clothing assets before restoring the LOD
 	TArray<ClothingAssetUtils::FClothingAssetMeshBinding> ClothingBindings;
-	ClothingAssetUtils::GetMeshClothingAssetBindings(SkeletalMesh, ClothingBindings);
-	for (ClothingAssetUtils::FClothingAssetMeshBinding& Binding : ClothingBindings)
-	{
-		//Unbind only the LOD we restore
-		if (Binding.LODIndex == LodIndex)
-		{
-			Binding.Asset->UnbindFromSkeletalMesh(SkeletalMesh, Binding.LODIndex);
-		}
-	}
+	FLODUtilities::UnbindClothingAndBackup(SkeletalMesh, ClothingBindings);
 
 	FSkeletalMeshLODModel ImportedBaseLODModel;
 	TMap<FString, TArray<FMorphTargetDelta>> ImportedBaseLODMorphTargetData;
-	SkeletalMesh->GetImportedModel()->OriginalReductionSourceMeshData[LodIndex]->LoadReductionData(ImportedBaseLODModel, ImportedBaseLODMorphTargetData);
+	SkeletalMesh->GetImportedModel()->OriginalReductionSourceMeshData[LodIndex]->LoadReductionData(ImportedBaseLODModel, ImportedBaseLODMorphTargetData, SkeletalMesh);
 	{
 		FSkeletalMeshUpdateContext UpdateContext;
 		UpdateContext.SkeletalMesh = SkeletalMesh;
 
-		TComponentReregisterContext<USkinnedMeshComponent> ReregisterContext;
-		if (bReregisterComponent)
-		{
-			SkeletalMesh->ReleaseResources();
-			SkeletalMesh->ReleaseResourcesFence.Wait();
-		}
 		//Copy the SkeletalMeshLODModel
 		SkeletalMesh->GetImportedModel()->LODModels[LodIndex] = ImportedBaseLODModel;
 		//Copy the morph target deltas
@@ -1196,20 +1145,7 @@ void FLODUtilities::RestoreSkeletalMeshLODImportedData(USkeletalMesh* SkeletalMe
 		SkeletalMesh->GetImportedModel()->OriginalReductionSourceMeshData[LodIndex]->EmptyBulkData();
 
 		//Put back the clothing for the restore LOD
-		for (ClothingAssetUtils::FClothingAssetMeshBinding& Binding : ClothingBindings)
-		{
-			if (LodIndex == Binding.LODIndex && SkeletalMesh->GetImportedModel()->LODModels.IsValidIndex(Binding.LODIndex) &&
-				SkeletalMesh->GetImportedModel()->LODModels[Binding.LODIndex].Sections.IsValidIndex(Binding.SectionIndex))
-			{
-				Binding.Asset->BindToSkeletalMesh(SkeletalMesh, Binding.LODIndex, Binding.SectionIndex, Binding.AssetInternalLodIndex);
-			}
-		}
-
-		if (bReregisterComponent)
-		{
-			SkeletalMesh->PostEditChange();
-			SkeletalMesh->InitResources();
-		}
+		FLODUtilities::RestoreClothingFromBackup(SkeletalMesh, ClothingBindings);
 
 		if (UpdateContext.OnLODChanged.IsBound())
 		{
@@ -1600,17 +1536,21 @@ void MatchVertexIndexUsingPosition(
 
 bool FLODUtilities::UpdateAlternateSkinWeights(USkeletalMesh* SkeletalMeshDest, const FName& ProfileNameDest, int32 LODIndexDest, FOverlappingThresholds OverlappingThresholds, bool ShouldImportNormals, bool ShouldImportTangents, bool bUseMikkTSpace, bool bComputeWeightedNormals)
 {
-	//Ensure log message only once
-	bool bNoMatchMsgDone = false;
-
 	//Grab all the destination structure
 	check(SkeletalMeshDest);
 	check(SkeletalMeshDest->GetImportedModel());
 	check(SkeletalMeshDest->GetImportedModel()->LODModels.IsValidIndex(LODIndexDest));
 	FSkeletalMeshLODModel& LODModelDest = SkeletalMeshDest->GetImportedModel()->LODModels[LODIndexDest];
+	return UpdateAlternateSkinWeights(LODModelDest, SkeletalMeshDest->GetName(), SkeletalMeshDest->RefSkeleton, ProfileNameDest, LODIndexDest, OverlappingThresholds, ShouldImportNormals, ShouldImportTangents, bUseMikkTSpace, bComputeWeightedNormals);
+}
+
+bool FLODUtilities::UpdateAlternateSkinWeights(FSkeletalMeshLODModel& LODModelDest, const FString SkeletalMeshName, FReferenceSkeleton& RefSkeleton, const FName& ProfileNameDest, int32 LODIndexDest, FOverlappingThresholds OverlappingThresholds, bool ShouldImportNormals, bool ShouldImportTangents, bool bUseMikkTSpace, bool bComputeWeightedNormals)
+{
+	//Ensure log message only once
+	bool bNoMatchMsgDone = false;
 	if (LODModelDest.RawSkeletalMeshBulkData.IsEmpty())
 	{
-		UE_LOG(LogLODUtilities, Error, TEXT("Failed to import Skin Weight Profile as the target skeletal mesh (%s) requires reimporting first."), SkeletalMeshDest ? *SkeletalMeshDest->GetName() : TEXT("NULL"));
+		UE_LOG(LogLODUtilities, Error, TEXT("Failed to import Skin Weight Profile as the target skeletal mesh (%s) requires reimporting first."), *SkeletalMeshName);
 		//Very old asset will not have this data, we cannot add alternate until the asset is reimported
 		return false;
 	}
@@ -1622,7 +1562,7 @@ bool FLODUtilities::UpdateAlternateSkinWeights(USkeletalMesh* SkeletalMeshDest,
 	int32 ProfileIndex = 0;
 	if (!ImportDataDest.AlternateInfluenceProfileNames.Find(ProfileNameDest.ToString(), ProfileIndex))
 	{
-		UE_LOG(LogLODUtilities, Error, TEXT("Failed to import Skin Weight Profile the alternate skinning imported source data is not available."), SkeletalMeshDest ? *SkeletalMeshDest->GetName() : TEXT("NULL"));
+		UE_LOG(LogLODUtilities, Error, TEXT("Failed to import Skin Weight Profile the alternate skinning imported source data is not available."), *SkeletalMeshName);
 		return false;
 	}
 	check(ImportDataDest.AlternateInfluences.IsValidIndex(ProfileIndex));
@@ -1634,7 +1574,7 @@ bool FLODUtilities::UpdateAlternateSkinWeights(USkeletalMesh* SkeletalMeshDest,
 
 	if (ImportDataDest.NumTexCoords <= 0 || ImportDataSrc.NumTexCoords <= 0)
 	{
-		UE_LOG(LogLODUtilities, Error, TEXT("Failed to import Skin Weight Profile as the target skeletal mesh (%s) or imported file does not contain UV coordinates."), SkeletalMeshDest ? *SkeletalMeshDest->GetName() : TEXT("NULL"));
+		UE_LOG(LogLODUtilities, Error, TEXT("Failed to import Skin Weight Profile as the target skeletal mesh (%s) or imported file does not contain UV coordinates."), *SkeletalMeshName);
 		return false;
 	}
 
@@ -1650,7 +1590,7 @@ bool FLODUtilities::UpdateAlternateSkinWeights(USkeletalMesh* SkeletalMeshDest,
 
 	//Create a map to remap source bone index to destination bone index
 	TMap<int32, int32> RemapBoneIndexSrcToDest;
-	FillRemapBoneIndexSrcToDest(ImportDataSrc, ImportDataDest, SkeletalMeshDest->GetName(), LODIndexDest, RemapBoneIndexSrcToDest);
+	FillRemapBoneIndexSrcToDest(ImportDataSrc, ImportDataDest, SkeletalMeshName, LODIndexDest, RemapBoneIndexSrcToDest);
 
 	//Map to get the vertex index source to a destination vertex match
 	TSortedMap<uint32, VertexMatchNameSpace::FVertexMatchResult> VertexIndexSrcToVertexIndexDestMatches;
@@ -1879,7 +1819,10 @@ bool FLODUtilities::UpdateAlternateSkinWeights(USkeletalMesh* SkeletalMeshDest,
 	TArray<FText> WarningMessages;
 	TArray<FName> WarningNames;
 	//Build the destination mesh with the Alternate influences, so the chunking is done properly.
-	bool bBuildSuccess = MeshUtilities.BuildSkeletalMesh(LODModelDest, SkeletalMeshDest->RefSkeleton, LODInfluencesDest, LODWedgesDest, LODFacesDest, LODPointsDest, LODPointToRawMapDest, BuildOptions, &WarningMessages, &WarningNames);
+	bool bBuildSuccess = MeshUtilities.BuildSkeletalMesh(LODModelDest, RefSkeleton, LODInfluencesDest, LODWedgesDest, LODFacesDest, LODPointsDest, LODPointToRawMapDest, BuildOptions, &WarningMessages, &WarningNames);
+	//Re-Apply the user section changes, the UserSectionsData is map to original section and should match the builded LODModel
+	LODModelDest.SyncronizeUserSectionsDataArray();
+
 	RegenerateAllImportSkinWeightProfileData(LODModelDest);
 	
 	return bBuildSuccess;
@@ -2004,7 +1947,8 @@ void FLODUtilities::RegenerateDependentLODs(USkeletalMesh* SkeletalMesh, int32 L
 {
 	FSkeletalMeshUpdateContext UpdateContext;
 	UpdateContext.SkeletalMesh = SkeletalMesh;
-	//Check the dependencies and regenerate the LODs acoording to it
+
+	//Check the dependencies and regenerate the LODs according to it
 	TArray<bool> LODDependencies;
 	int32 LODNumber = SkeletalMesh->GetLODNum();
 	LODDependencies.AddZeroed(LODNumber);
@@ -2022,19 +1966,20 @@ void FLODUtilities::RegenerateDependentLODs(USkeletalMesh* SkeletalMesh, int32 L
 	}
 	if (bRegenLODs)
 	{
-		TComponentReregisterContext<USkinnedMeshComponent> ReregisterContext;
-		SkeletalMesh->Modify();
-		SkeletalMesh->ReleaseResources();
-		SkeletalMesh->ReleaseResourcesFence.Wait();
-		for (int32 DependentLODIndex = LODIndex + 1; DependentLODIndex < LODNumber; ++DependentLODIndex)
+		FScopedSkeletalMeshPostEditChange ScopedPostEditChange(SkeletalMesh);
+		auto SimplifyDependentLODS = [&LODIndex, &LODNumber, &LODDependencies, &UpdateContext]()
 		{
-			if (LODDependencies[DependentLODIndex])
+			for (int32 DependentLODIndex = LODIndex + 1; DependentLODIndex < LODNumber; ++DependentLODIndex)
 			{
-				FLODUtilities::SimplifySkeletalMeshLOD(UpdateContext, DependentLODIndex, false);
+				if (LODDependencies[DependentLODIndex])
+				{
+					FLODUtilities::SimplifySkeletalMeshLOD(UpdateContext, DependentLODIndex, true);
+				}
 			}
-		}
-		SkeletalMesh->PostEditChange();
-		SkeletalMesh->InitResources();
+		};
+
+		SkeletalMesh->Modify();
+		SimplifyDependentLODS();
 	}
 }
 
@@ -2486,4 +2431,84 @@ void FLODUtilities::BuildMorphTargets(USkeletalMesh* BaseSkelMesh, FSkeletalMesh
 	}
 }
 
+void FLODUtilities::UnbindClothingAndBackup(USkeletalMesh* SkeletalMesh, TArray<ClothingAssetUtils::FClothingAssetMeshBinding>& ClothingBindings)
+{
+	for (int32 LODIndex = 0; LODIndex < SkeletalMesh->GetImportedModel()->LODModels.Num(); ++LODIndex)
+	{
+		TArray<ClothingAssetUtils::FClothingAssetMeshBinding> LODBindings;
+		UnbindClothingAndBackup(SkeletalMesh, LODBindings, LODIndex);
+		ClothingBindings.Append(LODBindings);
+	}
+}
+
+void FLODUtilities::UnbindClothingAndBackup(USkeletalMesh* SkeletalMesh, TArray<ClothingAssetUtils::FClothingAssetMeshBinding>& ClothingBindings, const int32 LODIndex)
+{
+	if (!SkeletalMesh->GetImportedModel()->LODModels.IsValidIndex(LODIndex))
+	{
+		return;
+	}
+	FSkeletalMeshLODModel& LODModel = SkeletalMesh->GetImportedModel()->LODModels[LODIndex];
+	//Store the clothBinding
+	ClothingAssetUtils::GetMeshClothingAssetBindings(SkeletalMesh, ClothingBindings, LODIndex);
+	//Unbind the Cloth for this LOD before we reduce it, we will put back the cloth after the reduction, if it still match the sections
+	for (ClothingAssetUtils::FClothingAssetMeshBinding& Binding : ClothingBindings)
+	{
+		if (Binding.LODIndex == LODIndex)
+		{
+			check(Binding.Asset);
+			Binding.Asset->UnbindFromSkeletalMesh(SkeletalMesh, Binding.LODIndex);
+			
+			//Use the UserSectionsData original section index, this will ensure we remap correctly the cloth if the reduction has change the number of sections
+			int32 OriginalDataSectionIndex = LODModel.Sections[Binding.SectionIndex].OriginalDataSectionIndex;
+			Binding.SectionIndex = OriginalDataSectionIndex;
+			
+			FSkelMeshSourceSectionUserData& SectionUserData = LODModel.UserSectionsData.FindChecked(OriginalDataSectionIndex);
+			SectionUserData.ClothingData.AssetGuid = FGuid();
+			SectionUserData.ClothingData.AssetLodIndex = INDEX_NONE;
+			SectionUserData.CorrespondClothAssetIndex = INDEX_NONE;
+		}
+	}
+}
+
+void FLODUtilities::RestoreClothingFromBackup(USkeletalMesh* SkeletalMesh, TArray<ClothingAssetUtils::FClothingAssetMeshBinding>& ClothingBindings)
+{
+	for (int32 LODIndex = 0; LODIndex < SkeletalMesh->GetImportedModel()->LODModels.Num(); ++LODIndex)
+	{
+		RestoreClothingFromBackup(SkeletalMesh, ClothingBindings, LODIndex);
+	}
+}
+
+void FLODUtilities::RestoreClothingFromBackup(USkeletalMesh* SkeletalMesh, TArray<ClothingAssetUtils::FClothingAssetMeshBinding>& ClothingBindings, const int32 LODIndex)
+{
+	if (!SkeletalMesh->GetImportedModel()->LODModels.IsValidIndex(LODIndex))
+	{
+		return;
+	}
+	FSkeletalMeshLODModel& LODModel = SkeletalMesh->GetImportedModel()->LODModels[LODIndex];
+	for (ClothingAssetUtils::FClothingAssetMeshBinding& Binding : ClothingBindings)
+	{
+		for (int32 SectionIndex = 0; SectionIndex < LODModel.Sections.Num(); ++SectionIndex)
+		{
+			if (LODModel.Sections[SectionIndex].OriginalDataSectionIndex != Binding.SectionIndex)
+			{
+				continue;
+			}
+			if (Binding.LODIndex == LODIndex)
+			{
+				check(Binding.Asset);
+				if (Binding.Asset->BindToSkeletalMesh(SkeletalMesh, Binding.LODIndex, SectionIndex, Binding.AssetInternalLodIndex))
+				{
+					//If successfull set back the section user data
+					FSkelMeshSourceSectionUserData& SectionUserData = LODModel.UserSectionsData.FindChecked(Binding.SectionIndex);
+					SectionUserData.CorrespondClothAssetIndex = LODModel.Sections[SectionIndex].CorrespondClothAssetIndex;
+					SectionUserData.ClothingData = LODModel.Sections[SectionIndex].ClothingData;
+				}
+			}
+			break;
+		}
+	}
+}
+
+
+
 #undef LOCTEXT_NAMESPACE // "LODUtilities"
\ No newline at end of file
diff --git a/Engine/Source/Editor/AnimGraph/Classes/AnimationGraphSchema.h b/Engine/Source/Editor/AnimGraph/Classes/AnimationGraphSchema.h
--- a/Engine/Source/Editor/AnimGraph/Classes/AnimationGraphSchema.h
+++ b/Engine/Source/Editor/AnimGraph/Classes/AnimationGraphSchema.h
@@ -13,6 +13,7 @@
 class FMenuBuilder;
 class UAnimationAsset;
 class UAnimBlueprint;
+class UPhysicsAsset;
 struct FBPInterfaceDescription;
 
 UCLASS(MinimalAPI)
@@ -77,6 +78,9 @@ class UAnimationGraphSchema : public UEdGraphSchema_K2
 	/** Spawn the correct node in the Animation Graph using the given AnimationAsset at the supplied location */
 	static void SpawnNodeFromAsset(UAnimationAsset* Asset, const FVector2D& GraphPosition, UEdGraph* Graph, UEdGraphPin* PinIfAvailable);
 
+	/** Spawn a rigid body node if we drop a physics asset on the graph */
+	static void SpawnRigidBodyNodeFromAsset(UPhysicsAsset* Asset, const FVector2D& GraphPosition, UEdGraph* Graph);
+
 	/** Update the specified node to a new asset */
 	static void UpdateNodeWithAsset(class UK2Node* K2Node, UAnimationAsset* Asset);
 
diff --git a/Engine/Source/Editor/AnimGraph/Private/AnimationGraphSchema.cpp b/Engine/Source/Editor/AnimGraph/Private/AnimationGraphSchema.cpp
--- a/Engine/Source/Editor/AnimGraph/Private/AnimationGraphSchema.cpp
+++ b/Engine/Source/Editor/AnimGraph/Private/AnimationGraphSchema.cpp
@@ -35,6 +35,7 @@
 #include "Animation/AnimMontage.h"
 #include "AnimGraphNode_LinkedInputPose.h"
 #include "AnimGraphNode_LinkedAnimLayer.h"
+#include "AnimGraphNode_RigidBody.h"
 
 #define LOCTEXT_NAMESPACE "AnimationGraphSchema"
 
@@ -410,6 +411,20 @@ void UAnimationGraphSchema::SpawnNodeFromAsset(UAnimationAsset* Asset, const FVe
 	}
 }
 
+void UAnimationGraphSchema::SpawnRigidBodyNodeFromAsset(UPhysicsAsset* Asset, const FVector2D& GraphPosition, UEdGraph* Graph)
+{
+	check(Graph);
+	check(Graph->GetSchema()->IsA(UAnimationGraphSchema::StaticClass()));
+	check(Asset);
+
+	FEdGraphSchemaAction_K2NewNode Action;
+
+	UAnimGraphNode_RigidBody* NewNode = NewObject<UAnimGraphNode_RigidBody>(GetTransientPackage());
+	NewNode->Node.OverridePhysicsAsset = Asset;
+	Action.NodeTemplate = NewNode;
+
+	Action.PerformAction(Graph, nullptr, GraphPosition);
+}
 
 void UAnimationGraphSchema::UpdateNodeWithAsset(UK2Node* K2Node, UAnimationAsset* Asset)
 {
@@ -431,18 +446,21 @@ void UAnimationGraphSchema::UpdateNodeWithAsset(UK2Node* K2Node, UAnimationAsset
 	}
 }
 
-
 void UAnimationGraphSchema::DroppedAssetsOnGraph( const TArray<FAssetData>& Assets, const FVector2D& GraphPosition, UEdGraph* Graph ) const 
 {
-	UAnimationAsset* Asset = FAssetData::GetFirstAsset<UAnimationAsset>(Assets);
-	if ((Asset != NULL) && (Graph != NULL))
+	if (Graph != NULL)
 	{
-		SpawnNodeFromAsset(Asset, GraphPosition, Graph, NULL);
+		if (UAnimationAsset* AnimationAsset = FAssetData::GetFirstAsset<UAnimationAsset>(Assets))
+		{
+			SpawnNodeFromAsset(AnimationAsset, GraphPosition, Graph, NULL);
+		}
+		else if (UPhysicsAsset* PhysicsAsset = FAssetData::GetFirstAsset<UPhysicsAsset>(Assets))
+		{
+			SpawnRigidBodyNodeFromAsset(PhysicsAsset, GraphPosition, Graph);
+		}
 	}
 }
 
-
-
 void UAnimationGraphSchema::DroppedAssetsOnNode(const TArray<FAssetData>& Assets, const FVector2D& GraphPosition, UEdGraphNode* Node) const
 {
 	UAnimationAsset* Asset = FAssetData::GetFirstAsset<UAnimationAsset>(Assets);
@@ -528,11 +546,10 @@ void UAnimationGraphSchema::GetAssetsPinHoverMessage(const TArray<FAssetData>& A
 
 void UAnimationGraphSchema::GetAssetsGraphHoverMessage(const TArray<FAssetData>& Assets, const UEdGraph* HoverGraph, FString& OutTooltipText, bool& OutOkIcon) const
 {
-	UAnimationAsset* Asset = FAssetData::GetFirstAsset<UAnimationAsset>(Assets);
-	if (Asset)
+	if (UAnimationAsset* AnimationAsset = FAssetData::GetFirstAsset<UAnimationAsset>(Assets))
 	{
 		UAnimBlueprint* AnimBlueprint = Cast<UAnimBlueprint>(FBlueprintEditorUtils::FindBlueprintForGraph(HoverGraph));
-		const bool bSkelMatch = (AnimBlueprint != NULL) && (AnimBlueprint->TargetSkeleton == Asset->GetSkeleton());
+		const bool bSkelMatch = (AnimBlueprint != NULL) && (AnimBlueprint->TargetSkeleton == AnimationAsset->GetSkeleton());
 		if (!bSkelMatch)
 		{
 			OutOkIcon = false;
@@ -549,6 +566,11 @@ void UAnimationGraphSchema::GetAssetsGraphHoverMessage(const TArray<FAssetData>&
 			OutTooltipText = TEXT("");
 		}
 	}
+	else if(UPhysicsAsset* PhysicsAsset = FAssetData::GetFirstAsset<UPhysicsAsset>(Assets))
+	{
+		OutOkIcon = true;
+		OutTooltipText = TEXT("");
+	}
 }
 
 void UAnimationGraphSchema::GetContextMenuActions(UToolMenu* Menu, UGraphNodeContextMenuContext* Context) const
diff --git a/Engine/Source/Editor/AnimationBlueprintEditor/Private/AnimationBlueprintEditor.cpp b/Engine/Source/Editor/AnimationBlueprintEditor/Private/AnimationBlueprintEditor.cpp
--- a/Engine/Source/Editor/AnimationBlueprintEditor/Private/AnimationBlueprintEditor.cpp
+++ b/Engine/Source/Editor/AnimationBlueprintEditor/Private/AnimationBlueprintEditor.cpp
@@ -310,7 +310,7 @@ void FAnimationBlueprintEditor::InitAnimationBlueprintEditor(const EToolkitMode:
 
 	// register customization of Slot node for this Animation Blueprint Editor
 	// this is so that you can open the manage window per Animation Blueprint Editor
-	PersonaModule.CustomizeSlotNodeDetails(Inspector->GetPropertyView().ToSharedRef(), FOnInvokeTab::CreateSP(this, &FAssetEditorToolkit::InvokeTab));
+	PersonaModule.CustomizeBlueprintEditorDetails(Inspector->GetPropertyView().ToSharedRef(), FOnInvokeTab::CreateSP(this, &FAssetEditorToolkit::InvokeTab));
 
 	if(bNewlyCreated && InAnimBlueprint->BlueprintType == BPTYPE_Interface)
 	{
@@ -396,7 +396,7 @@ void FAnimationBlueprintEditor::ExtendToolbar()
 		}
 	};
 
-    ToolbarExtender->AddToolBarExtension(
+	ToolbarExtender->AddToolBarExtension(
 		"Asset",
 		EExtensionHook::After,
 		GetToolkitCommands(),
@@ -1211,6 +1211,11 @@ FLinearColor FAnimationBlueprintEditor::GetWorldCentricTabColorScale() const
 	return FLinearColor( 0.5f, 0.25f, 0.35f, 0.5f );
 }
 
+IAnimationSequenceBrowser* FAnimationBlueprintEditor::GetAssetBrowser() const
+{
+	return SequenceBrowser.Pin().Get();
+}
+
 void FAnimationBlueprintEditor::OnActiveTabChanged( TSharedPtr<SDockTab> PreviouslyActive, TSharedPtr<SDockTab> NewlyActivated )
 {
 	if (!NewlyActivated.IsValid())
@@ -1483,7 +1488,7 @@ void FAnimationBlueprintEditor::HandleOpenNewAsset(UObject* InNewAsset)
 
 void FAnimationBlueprintEditor::AddReferencedObjects( FReferenceCollector& Collector )
 {
-    Collector.AddReferencedObject( EditorOptions );
+	Collector.AddReferencedObject( EditorOptions );
 }
 
 FAnimNode_Base* FAnimationBlueprintEditor::FindAnimNode(UAnimGraphNode_Base* AnimGraphNode) const
@@ -1535,7 +1540,7 @@ void FAnimationBlueprintEditor::OnSelectedNodesChangedImpl(const TSet<class UObj
 		}
 	}
 
-    bSelectRegularNode = false;
+	bSelectRegularNode = false;
 	for (FGraphPanelSelectionSet::TConstIterator It(NewSelection); It; ++It)
 	{
 		UEdGraphNode_Comment* SeqNode = Cast<UEdGraphNode_Comment>(*It);
@@ -1548,7 +1553,7 @@ void FAnimationBlueprintEditor::OnSelectedNodesChangedImpl(const TSet<class UObj
 		}
 	}
 
-    if (bHideUnrelatedNodes && !bLockNodeFadeState)
+	if (bHideUnrelatedNodes && !bLockNodeFadeState)
 	{
 		ResetAllNodesUnrelatedStates();
 
@@ -1813,5 +1818,10 @@ void FAnimationBlueprintEditor::HandlePreviewAnimBlueprintCompiled(UBlueprint* I
 	}
 }
 
+void FAnimationBlueprintEditor::HandleAnimationSequenceBrowserCreated(const TSharedRef<IAnimationSequenceBrowser>& InSequenceBrowser)
+{
+	SequenceBrowser = InSequenceBrowser;
+}
+
 #undef LOCTEXT_NAMESPACE
 
diff --git a/Engine/Source/Editor/AnimationBlueprintEditor/Private/AnimationBlueprintEditor.h b/Engine/Source/Editor/AnimationBlueprintEditor/Private/AnimationBlueprintEditor.h
--- a/Engine/Source/Editor/AnimationBlueprintEditor/Private/AnimationBlueprintEditor.h
+++ b/Engine/Source/Editor/AnimationBlueprintEditor/Private/AnimationBlueprintEditor.h
@@ -91,6 +91,7 @@ public:
 	/** IAnimationBlueprintEditor interface */
 	virtual const FEdGraphPinType& GetLastGraphPinTypeUsed() const override { return LastGraphPinType; }
 	virtual void SetLastGraphPinTypeUsed(const FEdGraphPinType& InType) override { LastGraphPinType = InType; }
+	virtual IAnimationSequenceBrowser* GetAssetBrowser() const override;
 
 	/** IHasPersonaToolkit interface */
 	virtual TSharedRef<class IPersonaToolkit> GetPersonaToolkit() const { return PersonaToolkit.ToSharedRef(); }
@@ -287,6 +288,8 @@ private:
 	 */
 	virtual void SaveEditorSettings();
 
+	void HandleAnimationSequenceBrowserCreated(const TSharedRef<IAnimationSequenceBrowser>& InSequenceBrowser);
+
 	/** The extender to pass to the level editor to extend it's window menu */
 	TSharedPtr<FExtender> MenuExtender;
 
@@ -305,6 +308,9 @@ private:
 	// selected anim graph node 
 	TWeakObjectPtr<class UAnimGraphNode_Base> SelectedAnimGraphNode;
 
+	/** Sequence Browser **/
+	TWeakPtr<class IAnimationSequenceBrowser> SequenceBrowser;
+
 	/** Delegate handle registered for when pin default values change */
 	FDelegateHandle OnPinDefaultValueChangedHandle;
 
diff --git a/Engine/Source/Editor/AnimationBlueprintEditor/Private/AnimationBlueprintEditorMode.cpp b/Engine/Source/Editor/AnimationBlueprintEditor/Private/AnimationBlueprintEditorMode.cpp
--- a/Engine/Source/Editor/AnimationBlueprintEditor/Private/AnimationBlueprintEditorMode.cpp
+++ b/Engine/Source/Editor/AnimationBlueprintEditor/Private/AnimationBlueprintEditorMode.cpp
@@ -132,7 +132,7 @@ FAnimationBlueprintEditorMode::FAnimationBlueprintEditorMode(const TSharedRef<FA
 	PersonaModule.RegisterPersonaViewportTabFactories(TabFactories, InAnimationBlueprintEditor, ViewportArgs);
 
 	TabFactories.RegisterFactory(PersonaModule.CreateAdvancedPreviewSceneTabFactory(InAnimationBlueprintEditor, InAnimationBlueprintEditor->GetPersonaToolkit()->GetPreviewScene()));
-	TabFactories.RegisterFactory(PersonaModule.CreateAnimationAssetBrowserTabFactory(InAnimationBlueprintEditor, InAnimationBlueprintEditor->GetPersonaToolkit(), FOnOpenNewAsset::CreateSP(&InAnimationBlueprintEditor.Get(), &FAnimationBlueprintEditor::HandleOpenNewAsset), FOnAnimationSequenceBrowserCreated(), true));
+	TabFactories.RegisterFactory(PersonaModule.CreateAnimationAssetBrowserTabFactory(InAnimationBlueprintEditor, InAnimationBlueprintEditor->GetPersonaToolkit(), FOnOpenNewAsset::CreateSP(&InAnimationBlueprintEditor.Get(), &FAnimationBlueprintEditor::HandleOpenNewAsset), FOnAnimationSequenceBrowserCreated::CreateSP(&InAnimationBlueprintEditor.Get(), &FAnimationBlueprintEditor::HandleAnimationSequenceBrowserCreated), true));
 	TabFactories.RegisterFactory(PersonaModule.CreateAnimBlueprintPreviewTabFactory(InAnimationBlueprintEditor, InAnimationBlueprintEditor->GetPersonaToolkit()->GetPreviewScene()));
 	TabFactories.RegisterFactory(PersonaModule.CreateAnimBlueprintAssetOverridesTabFactory(InAnimationBlueprintEditor, InAnimationBlueprintEditor->GetPersonaToolkit()->GetAnimBlueprint(), InAnimationBlueprintEditor->OnPostUndo));
 	TabFactories.RegisterFactory(PersonaModule.CreateSkeletonSlotNamesTabFactory(InAnimationBlueprintEditor, InAnimationBlueprintEditor->GetSkeletonTree()->GetEditableSkeleton(), InAnimationBlueprintEditor->OnPostUndo, FOnObjectSelected::CreateSP(&InAnimationBlueprintEditor.Get(), &FAnimationBlueprintEditor::HandleObjectSelected)));
diff --git a/Engine/Source/Editor/AnimationBlueprintEditor/Private/AnimationStateNodes/SGraphNodeAnimState.cpp b/Engine/Source/Editor/AnimationBlueprintEditor/Private/AnimationStateNodes/SGraphNodeAnimState.cpp
--- a/Engine/Source/Editor/AnimationBlueprintEditor/Private/AnimationStateNodes/SGraphNodeAnimState.cpp
+++ b/Engine/Source/Editor/AnimationBlueprintEditor/Private/AnimationStateNodes/SGraphNodeAnimState.cpp
@@ -83,37 +83,38 @@ void SGraphNodeAnimState::Construct(const FArguments& InArgs, UAnimStateNodeBase
 
 void SGraphNodeAnimState::GetStateInfoPopup(UEdGraphNode* GraphNode, TArray<FGraphInformationPopupInfo>& Popups)
 {
-	UAnimBlueprint* AnimBlueprint = Cast<UAnimBlueprint>(FBlueprintEditorUtils::FindBlueprintForNodeChecked(GraphNode));
-	check(AnimBlueprint);
-	UAnimInstance* ActiveObject = Cast<UAnimInstance>(AnimBlueprint->GetObjectBeingDebugged());
-	UAnimBlueprintGeneratedClass* Class = AnimBlueprint->GetAnimBlueprintGeneratedClass();
+	UAnimBlueprint* AnimBlueprint = Cast<UAnimBlueprint>(FBlueprintEditorUtils::FindBlueprintForNode(GraphNode));
+	if(AnimBlueprint)
+	{
+		UAnimInstance* ActiveObject = Cast<UAnimInstance>(AnimBlueprint->GetObjectBeingDebugged());
+		UAnimBlueprintGeneratedClass* Class = AnimBlueprint->GetAnimBlueprintGeneratedClass();
 
-	//FKismetNodeInfoContext* K2Context = (FKismetNodeInfoContext*)Context;
-	FLinearColor CurrentStateColor(1.f, 0.5f, 0.25f);
+		FLinearColor CurrentStateColor(1.f, 0.5f, 0.25f);
 
-	// Display various types of debug data
-	if ((ActiveObject != NULL) && (Class != NULL))
-	{
-		if (FStateMachineDebugData* DebugInfo = Class->GetAnimBlueprintDebugData().StateMachineDebugData.Find(GraphNode->GetGraph()))
+		// Display various types of debug data
+		if ((ActiveObject != NULL) && (Class != NULL))
 		{
-			if (Class->AnimNodeProperties.Num())
+			if (FStateMachineDebugData* DebugInfo = Class->GetAnimBlueprintDebugData().StateMachineDebugData.Find(GraphNode->GetGraph()))
 			{
-				UAnimationStateMachineGraph* TypedGraph = CastChecked<UAnimationStateMachineGraph>(GraphNode->GetGraph());
-
-				if (FAnimNode_StateMachine* CurrentInstance = Class->GetPropertyInstance<FAnimNode_StateMachine>(ActiveObject, TypedGraph->OwnerAnimGraphNode))  
+				if (Class->AnimNodeProperties.Num())
 				{
-					if (int32* pStateIndex = DebugInfo->NodeToStateIndex.Find(GraphNode))
+					UAnimationStateMachineGraph* TypedGraph = CastChecked<UAnimationStateMachineGraph>(GraphNode->GetGraph());
+
+					if (FAnimNode_StateMachine* CurrentInstance = Class->GetPropertyInstance<FAnimNode_StateMachine>(ActiveObject, TypedGraph->OwnerAnimGraphNode))  
 					{
-						const float Weight = CurrentInstance->GetStateWeight(*pStateIndex);
-						if (Weight > 0.0f)
+						if (int32* pStateIndex = DebugInfo->NodeToStateIndex.Find(GraphNode))
 						{
-							FString StateText = FString::Printf(TEXT("%.1f%%"), Weight * 100.0f);
-							if (*pStateIndex == CurrentInstance->GetCurrentState())
+							const float Weight = CurrentInstance->GetStateWeight(*pStateIndex);
+							if (Weight > 0.0f)
 							{
-								StateText += FString::Printf(TEXT("\nActive for %.2f secs"), CurrentInstance->GetCurrentStateElapsedTime());
-							}
+								FString StateText = FString::Printf(TEXT("%.1f%%"), Weight * 100.0f);
+								if (*pStateIndex == CurrentInstance->GetCurrentState())
+								{
+									StateText += FString::Printf(TEXT("\nActive for %.2f secs"), CurrentInstance->GetCurrentStateElapsedTime());
+								}
 
-							new (Popups) FGraphInformationPopupInfo(NULL, CurrentStateColor, StateText);
+								new (Popups) FGraphInformationPopupInfo(NULL, CurrentStateColor, StateText);
+							}
 						}
 					}
 				}
@@ -129,34 +130,34 @@ void SGraphNodeAnimState::GetNodeInfoPopups(FNodeInfoContext* Context, TArray<FG
 
 FSlateColor SGraphNodeAnimState::GetBorderBackgroundColor() const
 {
-	UAnimBlueprint* AnimBlueprint = Cast<UAnimBlueprint>(FBlueprintEditorUtils::FindBlueprintForNodeChecked(GraphNode));
-	check(AnimBlueprint);
-	UAnimInstance* ActiveObject = Cast<UAnimInstance>(AnimBlueprint->GetObjectBeingDebugged());
-	UAnimBlueprintGeneratedClass* Class = AnimBlueprint->GetAnimBlueprintGeneratedClass();
-
-	//FKismetNodeInfoContext* K2Context = (FKismetNodeInfoContext*)Context;
-
 	FLinearColor InactiveStateColor(0.08f, 0.08f, 0.08f);
 	FLinearColor ActiveStateColorDim(0.4f, 0.3f, 0.15f);
 	FLinearColor ActiveStateColorBright(1.f, 0.6f, 0.35f);
 
-	// Display various types of debug data
-	if ((ActiveObject != NULL) && (Class != NULL))
+	UAnimBlueprint* AnimBlueprint = Cast<UAnimBlueprint>(FBlueprintEditorUtils::FindBlueprintForNode(GraphNode));
+	if(AnimBlueprint)
 	{
-		if (FStateMachineDebugData* DebugInfo = Class->GetAnimBlueprintDebugData().StateMachineDebugData.Find(GraphNode->GetGraph()))
+		UAnimInstance* ActiveObject = Cast<UAnimInstance>(AnimBlueprint->GetObjectBeingDebugged());
+		UAnimBlueprintGeneratedClass* Class = AnimBlueprint->GetAnimBlueprintGeneratedClass();
+
+		// Display various types of debug data
+		if ((ActiveObject != NULL) && (Class != NULL))
 		{
-			if (Class->AnimNodeProperties.Num())
+			if (FStateMachineDebugData* DebugInfo = Class->GetAnimBlueprintDebugData().StateMachineDebugData.Find(GraphNode->GetGraph()))
 			{
-				UAnimationStateMachineGraph* TypedGraph = CastChecked<UAnimationStateMachineGraph>(GraphNode->GetGraph());
-
-				if (FAnimNode_StateMachine* CurrentInstance = Class->GetPropertyInstance<FAnimNode_StateMachine>(ActiveObject, TypedGraph->OwnerAnimGraphNode))
+				if (Class->AnimNodeProperties.Num())
 				{
-					if (int32* pStateIndex = DebugInfo->NodeToStateIndex.Find(GraphNode))
+					UAnimationStateMachineGraph* TypedGraph = CastChecked<UAnimationStateMachineGraph>(GraphNode->GetGraph());
+
+					if (FAnimNode_StateMachine* CurrentInstance = Class->GetPropertyInstance<FAnimNode_StateMachine>(ActiveObject, TypedGraph->OwnerAnimGraphNode))
 					{
-						const float Weight = CurrentInstance->GetStateWeight(*pStateIndex);
-						if (Weight > 0.0f)
+						if (int32* pStateIndex = DebugInfo->NodeToStateIndex.Find(GraphNode))
 						{
-							return FMath::Lerp<FLinearColor>(ActiveStateColorDim, ActiveStateColorBright, Weight);
+							const float Weight = CurrentInstance->GetStateWeight(*pStateIndex);
+							if (Weight > 0.0f)
+							{
+								return FMath::Lerp<FLinearColor>(ActiveStateColorDim, ActiveStateColorBright, Weight);
+							}
 						}
 					}
 				}
diff --git a/Engine/Source/Editor/AnimationBlueprintEditor/Public/IAnimationBlueprintEditor.h b/Engine/Source/Editor/AnimationBlueprintEditor/Public/IAnimationBlueprintEditor.h
--- a/Engine/Source/Editor/AnimationBlueprintEditor/Public/IAnimationBlueprintEditor.h
+++ b/Engine/Source/Editor/AnimationBlueprintEditor/Public/IAnimationBlueprintEditor.h
@@ -6,6 +6,8 @@
 #include "BlueprintEditor.h"
 #include "IHasPersonaToolkit.h"
 
+class IAnimationSequenceBrowser;
+
 class IAnimationBlueprintEditor : public FBlueprintEditor, public IHasPersonaToolkit
 {
 public:
@@ -14,4 +16,7 @@ public:
 
 	/** Set the last pin type we used to create a graph pin */
 	virtual void SetLastGraphPinTypeUsed(const FEdGraphPinType& InType) = 0;
+
+	/** Get the asset browser we host */
+	virtual IAnimationSequenceBrowser* GetAssetBrowser() const = 0;
 };
diff --git a/Engine/Source/Editor/AnimationEditor/Private/AnimationEditor.cpp b/Engine/Source/Editor/AnimationEditor/Private/AnimationEditor.cpp
--- a/Engine/Source/Editor/AnimationEditor/Private/AnimationEditor.cpp
+++ b/Engine/Source/Editor/AnimationEditor/Private/AnimationEditor.cpp
@@ -460,6 +460,11 @@ void FAnimationEditor::SetAnimationAsset(UAnimationAsset* AnimAsset)
 	HandleOpenNewAsset(AnimAsset);
 }
 
+IAnimationSequenceBrowser* FAnimationEditor::GetAssetBrowser() const
+{
+	return SequenceBrowser.Pin().Get();
+}
+
 void FAnimationEditor::HandleOpenNewAsset(UObject* InNewAsset)
 {
 	if (UAnimationAsset* NewAnimationAsset = Cast<UAnimationAsset>(InNewAsset))
diff --git a/Engine/Source/Editor/AnimationEditor/Private/AnimationEditor.h b/Engine/Source/Editor/AnimationEditor/Private/AnimationEditor.h
--- a/Engine/Source/Editor/AnimationEditor/Private/AnimationEditor.h
+++ b/Engine/Source/Editor/AnimationEditor/Private/AnimationEditor.h
@@ -56,6 +56,7 @@ public:
 
 	/** IAnimationEditor interface */
 	virtual void SetAnimationAsset(UAnimationAsset* AnimAsset) override;
+	virtual IAnimationSequenceBrowser* GetAssetBrowser() const override;
 
 	/** IHasPersonaToolkit interface */
 	virtual TSharedRef<class IPersonaToolkit> GetPersonaToolkit() const override { return PersonaToolkit.ToSharedRef(); }
diff --git a/Engine/Source/Editor/AnimationEditor/Public/IAnimationEditor.h b/Engine/Source/Editor/AnimationEditor/Public/IAnimationEditor.h
--- a/Engine/Source/Editor/AnimationEditor/Public/IAnimationEditor.h
+++ b/Engine/Source/Editor/AnimationEditor/Public/IAnimationEditor.h
@@ -7,10 +7,14 @@
 #include "IHasPersonaToolkit.h"
 
 class UAnimationAsset;
+class IAnimationSequenceBrowser;
 
 class IAnimationEditor : public FWorkflowCentricApplication, public IHasPersonaToolkit
 {
 public:
 	/** Set the animation asset of the editor. */
 	virtual void SetAnimationAsset(UAnimationAsset* AnimAsset) = 0;
+
+	/** Get the asset browser we host */
+	virtual IAnimationSequenceBrowser* GetAssetBrowser() const = 0;
 };
diff --git a/Engine/Source/Editor/BlueprintGraph/Private/BlueprintActionFilter.cpp b/Engine/Source/Editor/BlueprintGraph/Private/BlueprintActionFilter.cpp
--- a/Engine/Source/Editor/BlueprintGraph/Private/BlueprintActionFilter.cpp
+++ b/Engine/Source/Editor/BlueprintGraph/Private/BlueprintActionFilter.cpp
@@ -995,7 +995,7 @@ static bool BlueprintActionFilterImpl::IsIncompatibleImpureNode(FBlueprintAction
 
 	for (UEdGraph* Graph : FilterContext.Graphs)
 	{
-		if (const UEdGraphSchema_K2* K2Schema = GetDefault<UEdGraphSchema_K2>())
+		if (const UEdGraphSchema_K2* K2Schema = Cast<UEdGraphSchema_K2>(Graph->GetSchema()))
 		{
 			bAllowImpureNodes &= K2Schema->DoesGraphSupportImpureFunctions(Graph);
 		}
diff --git a/Engine/Source/Editor/ClothPainter/ClothPainter.Build.cs b/Engine/Source/Editor/ClothPainter/ClothPainter.Build.cs
--- a/Engine/Source/Editor/ClothPainter/ClothPainter.Build.cs
+++ b/Engine/Source/Editor/ClothPainter/ClothPainter.Build.cs
@@ -26,6 +26,7 @@ public class ClothPainter : ModuleRules
                 "SkeletalMeshEditor",                
                 "AssetRegistry",
                 "MeshPaint",
+                "ClothingSystemRuntimeCommon",
                 "ClothingSystemRuntimeNv",
                 "ClothingSystemRuntimeInterface",
                 "ClothingSystemEditorInterface",
diff --git a/Engine/Source/Editor/ClothPainter/Private/ClothMeshAdapter.cpp b/Engine/Source/Editor/ClothPainter/Private/ClothMeshAdapter.cpp
--- a/Engine/Source/Editor/ClothPainter/Private/ClothMeshAdapter.cpp
+++ b/Engine/Source/Editor/ClothPainter/Private/ClothMeshAdapter.cpp
@@ -3,7 +3,8 @@
 #include "ClothMeshAdapter.h"
 
 #include "ClothingAssetBase.h"
-#include "Assets/ClothingAsset.h"
+#include "ClothingAsset.h"
+#include "ClothPhysicalMeshDataNv.h" // for MaskTarget_PhysMesh
 
 #include "MeshPaintTypes.h"
 #include "Animation/DebugSkelMeshComponent.h"
@@ -54,9 +55,9 @@ bool FClothMeshPaintAdapter::Initialize()
 	bool bHaveAsset = false;
 	bool bBaseInit = false;
 
-	UClothingAsset* ClothingAsset = Cast<UClothingAsset>(SelectedAsset);
+	UClothingAssetCommon* ClothingAsset = Cast<UClothingAssetCommon>(SelectedAsset);
 
-	bHaveAsset = ClothingAsset && ClothingAsset->LodData.IsValidIndex(PaintingClothLODIndex);
+	bHaveAsset = ClothingAsset && ClothingAsset->ClothLodData.IsValidIndex(PaintingClothLODIndex);
 	bBaseInit = FBaseMeshPaintGeometryAdapter::Initialize();
 
 	return bHaveAsset && bBaseInit;
@@ -148,7 +149,10 @@ float FClothMeshPaintAdapter::GetMaxDistanceValue(int32 VertexIndex) const
 	{
 		if (VertexIndex >= Info.VertexStart && VertexIndex < Info.VertexEnd)
 		{
-			return Info.Asset->LodData[PaintingClothLODIndex].PhysicalMeshData.MaxDistances[VertexIndex - Info.VertexStart];			
+			const TArray<float>* MaxDistances = 
+				Info.Asset->ClothLodData[PaintingClothLODIndex]->PhysicalMeshData->GetFloatArray(
+				(uint32)MaskTarget_PhysMesh::MaxDistance);
+			return MaxDistances ? (*MaxDistances)[VertexIndex - Info.VertexStart] : 0.0f;
 		}
 	}
 
@@ -161,7 +165,12 @@ void FClothMeshPaintAdapter::SetMaxDistanceValue(int32 VertexIndex, float Value)
 	{
 		if (VertexIndex >= Info.VertexStart && VertexIndex < Info.VertexEnd)
 		{
-			Info.Asset->LodData[PaintingClothLODIndex].PhysicalMeshData.MaxDistances[VertexIndex - Info.VertexStart] = Value;
+			if (TArray<float>* MaxDistances =
+				Info.Asset->ClothLodData[PaintingClothLODIndex]->PhysicalMeshData->GetFloatArray(
+				(uint32)MaskTarget_PhysMesh::MaxDistance))
+			{
+				(*MaxDistances)[VertexIndex - Info.VertexStart] = Value;
+			}
 			break;
 		}
 	}
@@ -173,7 +182,10 @@ float FClothMeshPaintAdapter::GetBackstopDistanceValue(int32 VertexIndex) const
 	{
 		if (VertexIndex >= Info.VertexStart && VertexIndex < Info.VertexEnd)
 		{
-			return Info.Asset->LodData[PaintingClothLODIndex].PhysicalMeshData.BackstopDistances[VertexIndex - Info.VertexStart];
+			const TArray<float>* BackstopDistances = 
+				Info.Asset->ClothLodData[PaintingClothLODIndex]->PhysicalMeshData->GetFloatArray(
+				(uint32)MaskTarget_PhysMesh::BackstopDistance);
+			return BackstopDistances ? (*BackstopDistances)[VertexIndex - Info.VertexStart] : 0.0f;
 		}
 	}
 
@@ -186,7 +198,12 @@ void FClothMeshPaintAdapter::SetBackstopDistanceValue(int32 VertexIndex, float V
 	{
 		if (VertexIndex >= Info.VertexStart && VertexIndex < Info.VertexEnd)
 		{
-			Info.Asset->LodData[PaintingClothLODIndex].PhysicalMeshData.BackstopDistances[VertexIndex - Info.VertexStart] = Value;
+			if (TArray<float>* BackstopDistances =
+				Info.Asset->ClothLodData[PaintingClothLODIndex]->PhysicalMeshData->GetFloatArray(
+				(uint32)MaskTarget_PhysMesh::BackstopDistance))
+			{
+				(*BackstopDistances)[VertexIndex - Info.VertexStart] = Value;
+			}
 			break;
 		}
 	}
@@ -198,7 +215,10 @@ float FClothMeshPaintAdapter::GetBackstopRadiusValue(int32 VertexIndex) const
 	{
 		if (VertexIndex >= Info.VertexStart && VertexIndex < Info.VertexEnd)
 		{
-			return Info.Asset->LodData[PaintingClothLODIndex].PhysicalMeshData.BackstopRadiuses[VertexIndex - Info.VertexStart];
+			const TArray<float>* BackstopRadii = 
+				Info.Asset->ClothLodData[PaintingClothLODIndex]->PhysicalMeshData->GetFloatArray(
+				(uint32)MaskTarget_PhysMesh::BackstopRadius);
+			return BackstopRadii ? (*BackstopRadii)[VertexIndex - Info.VertexStart] : 0.0f;
 		}
 	}
 
@@ -211,7 +231,12 @@ void FClothMeshPaintAdapter::SetBackstopRadiusValue(int32 VertexIndex, float Val
 	{
 		if (VertexIndex >= Info.VertexStart && VertexIndex < Info.VertexEnd)
 		{
-			Info.Asset->LodData[PaintingClothLODIndex].PhysicalMeshData.BackstopRadiuses[VertexIndex - Info.VertexStart] = Value;
+			if (TArray<float>* BackstopRadii =
+				Info.Asset->ClothLodData[PaintingClothLODIndex]->PhysicalMeshData->GetFloatArray(
+				(uint32)MaskTarget_PhysMesh::BackstopRadius))
+			{
+				(*BackstopRadii)[VertexIndex - Info.VertexStart] = Value;
+			}
 			break;
 		}
 	}
@@ -252,14 +277,14 @@ void FClothMeshPaintAdapter::SetSelectedClothingAsset(const FGuid& InAssetGuid,
 	{
 		for(UClothingAssetBase* Asset : ReferencedSkeletalMesh->MeshClothingAssets)
 		{
-			UClothingAsset* ConcreteAsset = CastChecked<UClothingAsset>(Asset);
+			UClothingAssetCommon* ConcreteAsset = CastChecked<UClothingAssetCommon>(Asset);
 			if(ConcreteAsset->GetAssetGuid() == InAssetGuid)
 			{
 				if(ConcreteAsset->IsValidLod(InAssetLod))
 				{
-					FClothLODData& LodData = ConcreteAsset->LodData[InAssetLod];
+					UClothLODDataBase* LodData = ConcreteAsset->ClothLodData[InAssetLod];
 
-					if(LodData.ParameterMasks.IsValidIndex(InMaskIndex))
+					if(LodData->ParameterMasks.IsValidIndex(InMaskIndex))
 					{
 						PaintingClothLODIndex = InAssetLod;
 						PaintingClothMaskIndex = InMaskIndex;
@@ -307,12 +332,13 @@ bool FClothMeshPaintAdapter::InitializeVertexData()
 		{
 			if(DebugComponent->SkinnedSelectedClothingPositions.Num() > 0)
 			{
-				UClothingAsset* ConcreteAsset = CastChecked<UClothingAsset>(SelectedAsset);
-				const FClothLODData& LODData = ConcreteAsset->LodData[PaintingClothLODIndex];
-				const FClothPhysicalMeshData& MeshData = LODData.PhysicalMeshData;
+				UClothingAssetCommon* ConcreteAsset = CastChecked<UClothingAssetCommon>(SelectedAsset);
+				const UClothLODDataBase* LODData = ConcreteAsset->ClothLodData[PaintingClothLODIndex];
+				check(LODData->PhysicalMeshData);
+				const UClothPhysicalMeshDataBase* MeshData = LODData->PhysicalMeshData;
 
 				MeshVertices.Append(DebugComponent->SkinnedSelectedClothingPositions);
-				MeshIndices.Append(MeshData.Indices);
+				MeshIndices.Append(MeshData->Indices);
 
 				for(int32 Index = IndexOffset; Index < MeshIndices.Num(); ++Index)
 				{
@@ -323,8 +349,8 @@ bool FClothMeshPaintAdapter::InitializeVertexData()
 				Info.IndexStart = IndexOffset;
 				Info.VertexStart = VertexOffset;
 
-				IndexOffset += MeshData.Indices.Num();
-				VertexOffset += MeshData.Vertices.Num();
+				IndexOffset += MeshData->Indices.Num();
+				VertexOffset += MeshData->Vertices.Num();
 
 				Info.IndexEnd = IndexOffset;
 				Info.VertexEnd = VertexOffset;
@@ -365,13 +391,13 @@ bool FClothMeshPaintAdapter::InitializeVertexData()
 	return true;
 }
 
-FClothParameterMask_PhysMesh* FClothMeshPaintAdapter::GetCurrentMask() const
+FPointWeightMap* FClothMeshPaintAdapter::GetCurrentMask() const
 {
 	if(HasValidSelection())
 	{
-		UClothingAsset* ConcreteAsset = CastChecked<UClothingAsset>(SelectedAsset);
+		UClothingAssetCommon* ConcreteAsset = CastChecked<UClothingAssetCommon>(SelectedAsset);
 
-		return &ConcreteAsset->LodData[PaintingClothLODIndex].ParameterMasks[PaintingClothMaskIndex];
+		return &ConcreteAsset->ClothLodData[PaintingClothLODIndex]->ParameterMasks[PaintingClothMaskIndex];
 	}
 
 	return nullptr;
@@ -381,11 +407,11 @@ bool FClothMeshPaintAdapter::HasValidSelection() const
 {
 	if(SelectedAsset)
 	{
-		UClothingAsset* ConcreteAsset = Cast<UClothingAsset>(SelectedAsset);
+		UClothingAssetCommon* ConcreteAsset = Cast<UClothingAssetCommon>(SelectedAsset);
 
 		// Only valid selection if we have a valid asset, asset LOD and a mask
-		if(ConcreteAsset->LodData.IsValidIndex(PaintingClothLODIndex) &&
-			ConcreteAsset->LodData[PaintingClothLODIndex].ParameterMasks.IsValidIndex(PaintingClothMaskIndex))
+		if(ConcreteAsset->ClothLodData.IsValidIndex(PaintingClothLODIndex) &&
+			ConcreteAsset->ClothLodData[PaintingClothLODIndex]->ParameterMasks.IsValidIndex(PaintingClothMaskIndex))
 		{
 			return true;
 		}
diff --git a/Engine/Source/Editor/ClothPainter/Private/ClothMeshAdapter.h b/Engine/Source/Editor/ClothPainter/Private/ClothMeshAdapter.h
--- a/Engine/Source/Editor/ClothPainter/Private/ClothMeshAdapter.h
+++ b/Engine/Source/Editor/ClothPainter/Private/ClothMeshAdapter.h
@@ -7,11 +7,11 @@
 #include "BaseMeshPaintGeometryAdapter.h"
 #include "IMeshPaintGeometryAdapterFactory.h"
 
-struct FClothParameterMask_PhysMesh;
+struct FPointWeightMap;
 class FReferenceCollector;
 class UMeshComponent;
 class UClothingAssetBase;
-class UClothingAsset;
+class UClothingAssetCommon;
 class USkeletalMeshComponent;
 class USkeletalMesh;
 
@@ -33,7 +33,7 @@ protected:
 		TArray<TArray<int32>> NeighborMap;
 
 		/** The actual clothing asset relating to this data */
-		UClothingAsset* Asset;
+		UClothingAssetCommon* Asset;
 	};
 public:
 	static void InitializeAdapterGlobals() {}
@@ -84,7 +84,7 @@ public:
 	const TArray<int32>* GetVertexNeighbors(int32 InVertexIndex) const;
 
 	/** Get the current mask we're editing */
-	FClothParameterMask_PhysMesh* GetCurrentMask() const;
+	FPointWeightMap* GetCurrentMask() const;
 
 protected:
 
diff --git a/Engine/Source/Editor/ClothPainter/Private/ClothPaintSettings.h b/Engine/Source/Editor/ClothPainter/Private/ClothPaintSettings.h
--- a/Engine/Source/Editor/ClothPainter/Private/ClothPaintSettings.h
+++ b/Engine/Source/Editor/ClothPainter/Private/ClothPaintSettings.h
@@ -7,9 +7,9 @@
 
 #include "ClothPaintSettings.generated.h"
 
-class UClothingAsset;
+class UClothingAssetCommon;
 
-DECLARE_MULTICAST_DELEGATE_ThreeParams(FOnClothingAssetSelectionChangedMulticaster, UClothingAsset*, int32, int32);
+DECLARE_MULTICAST_DELEGATE_ThreeParams(FOnClothingAssetSelectionChangedMulticaster, UClothingAssetCommon*, int32, int32);
 typedef FOnClothingAssetSelectionChangedMulticaster::FDelegate FOnClothingAssetSelectionChanged;
 
 UENUM()
@@ -84,7 +84,7 @@ public:
 
 	/** Array of Clothing assets */
 	UPROPERTY()
-	TArray<UClothingAsset*> ClothingAssets;
+	TArray<UClothingAssetCommon*> ClothingAssets;
 
 	/** Whether to flip normals on the mesh preview */
 	UPROPERTY(EditAnywhere, AdvancedDisplay, Category = View)
diff --git a/Engine/Source/Editor/ClothPainter/Private/ClothPaintSettingsCustomization.h b/Engine/Source/Editor/ClothPainter/Private/ClothPaintSettingsCustomization.h
--- a/Engine/Source/Editor/ClothPainter/Private/ClothPaintSettingsCustomization.h
+++ b/Engine/Source/Editor/ClothPainter/Private/ClothPaintSettingsCustomization.h
@@ -8,7 +8,7 @@
 #include "Widgets/Views/SListView.h"
 
 class FClothPainter;
-class UClothingAsset;
+class UClothingAssetCommon;
 class UClothPainterSettings;
 class FClothPaintToolBase;
 
diff --git a/Engine/Source/Editor/ClothPainter/Private/ClothPainter.cpp b/Engine/Source/Editor/ClothPainter/Private/ClothPainter.cpp
--- a/Engine/Source/Editor/ClothPainter/Private/ClothPainter.cpp
+++ b/Engine/Source/Editor/ClothPainter/Private/ClothPainter.cpp
@@ -6,7 +6,7 @@
 #include "MeshPaintAdapterFactory.h"
 #include "MeshPaintHelpers.h"
 
-#include "Assets/ClothingAsset.h"
+#include "ClothingAsset.h"
 #include "Animation/DebugSkelMeshComponent.h"
 #include "ClothPaintSettings.h"
 #include "ClothMeshAdapter.h"
@@ -179,7 +179,7 @@ void FClothPainter::RefreshClothingAssets()
 	{
 		for(UClothingAssetBase* BaseClothingAsset : Mesh->MeshClothingAssets)
 		{
-			if(UClothingAsset* ActualAsset = Cast<UClothingAsset>(BaseClothingAsset))
+			if(UClothingAssetCommon* ActualAsset = Cast<UClothingAssetCommon>(BaseClothingAsset))
 			{
 				PaintSettings->ClothingAssets.AddUnique(ActualAsset);
 			}
@@ -213,7 +213,7 @@ void FClothPainter::RecalculateAutoViewRange()
 	}
 
 	TSharedPtr<FClothMeshPaintAdapter> ClothAdapter = StaticCastSharedPtr<FClothMeshPaintAdapter>(Adapter);
-	FClothParameterMask_PhysMesh* CurrentMask = ClothAdapter->GetCurrentMask();
+	FPointWeightMap* CurrentMask = ClothAdapter->GetCurrentMask();
 
 	if(UClothPainterSettings* PainterSettings = Cast<UClothPainterSettings>(GetPainterSettings()))
 	{
@@ -264,7 +264,7 @@ void FClothPainter::Tick(FEditorViewportClient* ViewportClient, float DeltaTime)
 		if(bShouldSimulate)
 		{
 			// Need to re-apply our masks here, as they have likely been edited
-			for(UClothingAsset* Asset : PaintSettings->ClothingAssets)
+			for(UClothingAssetCommon* Asset : PaintSettings->ClothingAssets)
 			{
 				if(Asset)
 				{
@@ -417,7 +417,7 @@ void FClothPainter::Render(const FSceneView* View, FViewport* Viewport, FPrimiti
 				TArray<TPair<int32, FVector>> VertexData;
 				ClothAdapter->GetInfluencedVertexData(ComponentSpaceSquaredBrushRadius, ComponentSpaceBrushPosition, ComponentSpaceCameraPosition, BrushSettings->bOnlyFrontFacingTriangles, VertexData);
 				
-				FClothParameterMask_PhysMesh* CurrentMask = ClothAdapter->GetCurrentMask();
+				FPointWeightMap* CurrentMask = ClothAdapter->GetCurrentMask();
 
 				if(CurrentMask && VertexData.Num() > 0)
 				{
@@ -539,7 +539,7 @@ float FClothPainter::GetPropertyValue(int32 VertexIndex)
 {
 	FClothMeshPaintAdapter* ClothAdapter = (FClothMeshPaintAdapter*)Adapter.Get();
 
-	if(FClothParameterMask_PhysMesh* Mask = ClothAdapter->GetCurrentMask())
+	if(FPointWeightMap* Mask = ClothAdapter->GetCurrentMask())
 	{
 		return Mask->GetValue(VertexIndex);
 	}
@@ -551,21 +551,21 @@ void FClothPainter::SetPropertyValue(int32 VertexIndex, const float Value)
 {
 	FClothMeshPaintAdapter* ClothAdapter = (FClothMeshPaintAdapter*)Adapter.Get();
 
-	if(FClothParameterMask_PhysMesh* Mask = ClothAdapter->GetCurrentMask())
+	if(FPointWeightMap* Mask = ClothAdapter->GetCurrentMask())
 	{
 		Mask->SetValue(VertexIndex, Value);
 	}
 }
 
-void FClothPainter::OnAssetSelectionChanged(UClothingAsset* InNewSelectedAsset, int32 InAssetLod, int32 InMaskIndex)
+void FClothPainter::OnAssetSelectionChanged(UClothingAssetCommon* InNewSelectedAsset, int32 InAssetLod, int32 InMaskIndex)
 {
 	TSharedPtr<FClothMeshPaintAdapter> ClothAdapter = StaticCastSharedPtr<FClothMeshPaintAdapter>(Adapter);
 	if(ClothAdapter.IsValid() && InNewSelectedAsset && InNewSelectedAsset->IsValidLod(InAssetLod))
 	{
 		// Validate the incoming parameters, to make sure we only set a selection if we're going
 		// to get a valid paintable surface
-		if(InNewSelectedAsset->LodData.IsValidIndex(InAssetLod) &&
-			 InNewSelectedAsset->LodData[InAssetLod].ParameterMasks.IsValidIndex(InMaskIndex))
+		if(InNewSelectedAsset->ClothLodData.IsValidIndex(InAssetLod) &&
+			 InNewSelectedAsset->ClothLodData[InAssetLod]->ParameterMasks.IsValidIndex(InMaskIndex))
 		{
 			const FGuid NewGuid = InNewSelectedAsset->GetAssetGuid();
 			SkeletalMeshComponent->SetMeshSectionVisibilityForCloth(SkeletalMeshComponent->SelectedClothingGuidForPainting, true);
diff --git a/Engine/Source/Editor/ClothPainter/Private/ClothPainter.h b/Engine/Source/Editor/ClothPainter/Private/ClothPainter.h
--- a/Engine/Source/Editor/ClothPainter/Private/ClothPainter.h
+++ b/Engine/Source/Editor/ClothPainter/Private/ClothPainter.h
@@ -10,7 +10,7 @@
 class SClothPaintWidget;
 class UClothPainterSettings;
 class UPaintBrushSettings;
-class UClothingAsset;
+class UClothingAssetCommon;
 class UDebugSkelMeshComponent;
 enum class EPaintableClothProperty;
 class FClothPaintToolBase;
@@ -78,7 +78,7 @@ public:
 	TSharedPtr<IMeshPaintGeometryAdapter> GetAdapter() const { return Adapter; }
 
 	/** When a different clothing asset is selected in the UI the painter should refresh the adapter */
-	void OnAssetSelectionChanged(UClothingAsset* InNewSelectedAsset, int32 InAssetLod, int32 MaskIndex);
+	void OnAssetSelectionChanged(UClothingAssetCommon* InNewSelectedAsset, int32 InAssetLod, int32 MaskIndex);
 	void OnAssetMaskSelectionChanged()
 	{};
 
diff --git a/Engine/Source/Editor/ClothPainter/Private/ClothingPaintEditMode.cpp b/Engine/Source/Editor/ClothPainter/Private/ClothingPaintEditMode.cpp
--- a/Engine/Source/Editor/ClothPainter/Private/ClothingPaintEditMode.cpp
+++ b/Engine/Source/Editor/ClothPainter/Private/ClothingPaintEditMode.cpp
@@ -12,7 +12,7 @@
 #include "ClothingAssetBase.h"
 #include "ComponentRecreateRenderStateContext.h"
 #include "IPersonaToolkit.h"
-#include "Assets/ClothingAsset.h"
+#include "ClothingAsset.h"
 #include "EditorViewportClient.h"
 #include "AssetViewerSettings.h"
 #include "Editor/EditorPerProjectUserSettings.h"
@@ -92,7 +92,7 @@ void FClothingPaintEditMode::Exit()
 			{
 				for(UClothingAssetBase* AssetBase : SkelMesh->MeshClothingAssets)
 				{
-					UClothingAsset* ConcreteAsset = CastChecked<UClothingAsset>(AssetBase);
+					UClothingAssetCommon* ConcreteAsset = CastChecked<UClothingAssetCommon>(AssetBase);
 					ConcreteAsset->ApplyParameterMasks();
 				}
 			}
