The dependent LODs are now generate in parallel when building the asset

diff --git a/Engine/Plugins/Experimental/SkeletalReduction/Source/Private/SkeletalMeshReductionPlugin.cpp b/Engine/Plugins/Experimental/SkeletalReduction/Source/Private/SkeletalMeshReductionPlugin.cpp
--- a/Engine/Plugins/Experimental/SkeletalReduction/Source/Private/SkeletalMeshReductionPlugin.cpp
+++ b/Engine/Plugins/Experimental/SkeletalReduction/Source/Private/SkeletalMeshReductionPlugin.cpp
@@ -310,6 +310,12 @@ public:
 	/** IMeshReductionModule interface.*/
 	virtual class IMeshReduction* GetSkeletalMeshReductionInterface() override
 	{
+		if (IsInGameThread())
+		{
+			//Load dependent modules in case the reduction is call later during a multi threaded call
+			FModuleManager::Get().LoadModuleChecked<IMeshBoneReductionModule>("MeshBoneReduction");
+			FModuleManager::Get().LoadModuleChecked<IMeshUtilities>("MeshUtilities");
+		}
 		return &SkeletalMeshReducer;
 	}
 
@@ -620,7 +626,7 @@ void FQuadricSkeletalMeshReduction::ConvertToFSkinnedSkeletalMesh( const FSkelet
 			if (bHasBadNTB)
 			{
 				NumBadNTBSpace++;
-		}
+			}
 		}
 
 		// Report any error with invalid bone weights
@@ -1971,7 +1977,18 @@ void FQuadricSkeletalMeshReduction::ReduceSkeletalMesh(USkeletalMesh& SkeletalMe
 				GeoImportVersion = Old->RawSkeletalMeshBulkData.GeoImportVersion;
 				SkinningImportVersion = Old->RawSkeletalMeshBulkData.SkinningImportVersion;
 			}
-			delete Old;
+			//If the delegate is not bound 
+			if (!Settings.OnDeleteLODModelDelegate.IsBound())
+			{
+				//If not in game thread we should never delete a structure containing bulkdata since it can crash when the bulkdata is detach from the archive
+				//Use the delegate and delete the pointer in the main thread if you reduce in other thread then game thread (main thread).
+				check(IsInGameThread());
+				delete Old;
+			}
+			else
+			{
+				Settings.OnDeleteLODModelDelegate.Execute(Old);
+			}
 		}
 		else if(bReducingSourceModel)
 		{
diff --git a/Engine/Source/Developer/SkeletalMeshUtilitiesCommon/Private/LODUtilities.cpp b/Engine/Source/Developer/SkeletalMeshUtilitiesCommon/Private/LODUtilities.cpp
--- a/Engine/Source/Developer/SkeletalMeshUtilitiesCommon/Private/LODUtilities.cpp
+++ b/Engine/Source/Developer/SkeletalMeshUtilitiesCommon/Private/LODUtilities.cpp
@@ -29,7 +29,7 @@
 #include "IMeshReductionManagerModule.h"
 #include "Animation/SkinWeightProfile.h"
 
-#include "Modules/ModuleManager.h"
+#include "Async/ParallelFor.h"
 
 IMPLEMENT_MODULE(FDefaultModuleImpl, SkeletalMeshUtilitiesCommon)
 
@@ -232,7 +232,8 @@ bool FLODUtilities::RegenerateLOD(USkeletalMesh* SkeletalMesh, int32 NewLODCount
 		// we need to add more
 		else if (LODCount > CurrentNumLODs)
 		{
-			// Only create new skeletal mesh LOD level entries
+			// Only create new skeletal mesh LOD level entries, we cannot multi thread since the LOD will be create here
+			//TArray are not thread safe.
 			for (int32 LODIdx = CurrentNumLODs; LODIdx < LODCount; LODIdx++)
 			{
 				// if no previous setting found, it will use default setting. 
@@ -968,6 +969,7 @@ void FLODUtilities::SimplifySkeletalMeshLOD( USkeletalMesh* SkeletalMesh, int32
 		return;
 	}
 
+	if (IsInGameThread())
 	{
 		FFormatNamedArguments Args;
 		Args.Add(TEXT("DesiredLOD"), DesiredLOD);
@@ -1079,7 +1081,10 @@ void FLODUtilities::SimplifySkeletalMeshLOD( USkeletalMesh* SkeletalMesh, int32
 		FLODUtilities::RestoreClothingFromBackup(SkeletalMesh, ClothingBindings, DesiredLOD);
 	}
 
-	GWarn->EndSlowTask();
+	if (IsInGameThread())
+	{
+		GWarn->EndSlowTask();
+	}
 }
 
 void FLODUtilities::SimplifySkeletalMeshLOD(FSkeletalMeshUpdateContext& UpdateContext, int32 DesiredLOD, bool bRestoreClothing /*= false*/)
@@ -1942,41 +1947,93 @@ void FLODUtilities::RegenerateAllImportSkinWeightProfileData(FSkeletalMeshLODMod
 
 void FLODUtilities::RegenerateDependentLODs(USkeletalMesh* SkeletalMesh, int32 LODIndex)
 {
-	FSkeletalMeshUpdateContext UpdateContext;
-	UpdateContext.SkeletalMesh = SkeletalMesh;
-
-	//Check the dependencies and regenerate the LODs according to it
-	TArray<bool> LODDependencies;
 	int32 LODNumber = SkeletalMesh->GetLODNum();
-	LODDependencies.AddZeroed(LODNumber);
-	bool bRegenLODs = false;
-	LODDependencies[LODIndex] = true;
+	TMap<int32, TArray<int32>> Dependencies;
+	TBitArray<> DependentLOD;
+	DependentLOD.Init(false, LODNumber);
+	DependentLOD[LODIndex] = true;
 	for (int32 DependentLODIndex = LODIndex + 1; DependentLODIndex < LODNumber; ++DependentLODIndex)
 	{
 		const FSkeletalMeshLODInfo* LODInfo = SkeletalMesh->GetLODInfo(DependentLODIndex);
-		if (LODInfo && LODInfo->bHasBeenSimplified && LODDependencies[LODInfo->ReductionSettings.BaseLOD])
+		//Only add active reduction LOD that are not inline reducted (inline mean they do not depend on LODIndex)
+		if (LODInfo && (SkeletalMesh->IsReductionActive(DependentLODIndex) || LODInfo->bHasBeenSimplified) && DependentLODIndex > LODInfo->ReductionSettings.BaseLOD && DependentLOD[LODInfo->ReductionSettings.BaseLOD])
 		{
-			LODDependencies[DependentLODIndex] = true;
-			bRegenLODs = true;
-
+			TArray<int32>& LODDependencies = Dependencies.FindOrAdd(LODInfo->ReductionSettings.BaseLOD);
+			LODDependencies.Add(DependentLODIndex);
+			DependentLOD[DependentLODIndex] = true;
 		}
 	}
-	if (bRegenLODs)
+	if (Dependencies.Contains(LODIndex))
 	{
+		//Load the necessary module before going multithreaded
+		IMeshReductionModule& ReductionModule = FModuleManager::Get().LoadModuleChecked<IMeshReductionModule>("MeshReductionInterface");
+		//This will load all necessary module before kicking the multi threaded reduction
+		IMeshReduction* MeshReduction = ReductionModule.GetSkeletalMeshReductionInterface();
+		check(MeshReduction && MeshReduction->IsSupported());
+
 		FScopedSkeletalMeshPostEditChange ScopedPostEditChange(SkeletalMesh);
-		auto SimplifyDependentLODS = [&LODIndex, &LODNumber, &LODDependencies, &UpdateContext]()
 		{
-			for (int32 DependentLODIndex = LODIndex + 1; DependentLODIndex < LODNumber; ++DependentLODIndex)
+			FFormatNamedArguments Args;
+			Args.Add(TEXT("DesiredLOD"), LODIndex);
+			Args.Add(TEXT("SkeletalMeshName"), FText::FromString(SkeletalMesh->GetName()));
+			const FText StatusUpdate = FText::Format(NSLOCTEXT("UnrealEd", "MeshSimp_GeneratingDependentLODs_F", "Generating All Dependent LODs from LOD {DesiredLOD} for {SkeletalMeshName}..."), Args);
+			GWarn->BeginSlowTask(StatusUpdate, true);
+		}
+		for (const auto& Kvp : Dependencies)
+		{
+			SkeletalMesh->Modify();
+			//Use a TQueue which is thread safe, this Queue will be fill by some delegate call from other threads
+			TQueue<FSkeletalMeshLODModel*> LODModelReplaceByReduction;
+
+			const TArray<int32>& DependentLODs = Kvp.Value;
+			//Clothing do not play well with multithread, backup it here. Also bind the LODModel delete delegates
+			TMap<int32, TArray<ClothingAssetUtils::FClothingAssetMeshBinding>> PerLODClothingBindings;
+			for (int32 DependentLODIndex : DependentLODs)
 			{
-				if (LODDependencies[DependentLODIndex])
+				TArray<ClothingAssetUtils::FClothingAssetMeshBinding>& ClothingBindings = PerLODClothingBindings.FindOrAdd(DependentLODIndex);
+				FLODUtilities::UnbindClothingAndBackup(SkeletalMesh, ClothingBindings, DependentLODIndex);
+
+				const FSkeletalMeshLODInfo* LODInfo = SkeletalMesh->GetLODInfo(DependentLODIndex);
+				check(LODInfo);
+				LODInfo->ReductionSettings.OnDeleteLODModelDelegate.BindLambda([&LODModelReplaceByReduction](FSkeletalMeshLODModel* ReplacedLODModel)
 				{
-					FLODUtilities::SimplifySkeletalMeshLOD(UpdateContext, DependentLODIndex, true);
+					LODModelReplaceByReduction.Enqueue(ReplacedLODModel);
+				});
+			}
+
+			//Reduce all dependent LOD in same time
+			ParallelFor(DependentLODs.Num(), [&DependentLODs, &SkeletalMesh](int32 IterationIndex)
+			{
+				check(DependentLODs.IsValidIndex(IterationIndex));
+				int32 DependentLODIndex = DependentLODs[IterationIndex];
+				check(SkeletalMesh->GetLODInfo(DependentLODIndex)); //We cannot add a LOD when reducing with multi thread, so check we already have one
+				FLODUtilities::SimplifySkeletalMeshLOD(SkeletalMesh, DependentLODIndex, false);
+			});
+
+			//Restore the clothings and unbind the delegates
+			for (int32 DependentLODIndex : DependentLODs)
+			{
+				TArray<ClothingAssetUtils::FClothingAssetMeshBinding>& ClothingBindings = PerLODClothingBindings.FindChecked(DependentLODIndex);
+				FLODUtilities::RestoreClothingFromBackup(SkeletalMesh, ClothingBindings);
+
+				FSkeletalMeshLODInfo* LODInfo = SkeletalMesh->GetLODInfo(DependentLODIndex);
+				check(LODInfo);
+				LODInfo->ReductionSettings.OnDeleteLODModelDelegate.Unbind();
+			}
+
+			while (!LODModelReplaceByReduction.IsEmpty())
+			{
+				FSkeletalMeshLODModel* ReplacedLODModel = nullptr;
+				LODModelReplaceByReduction.Dequeue(ReplacedLODModel);
+				if (ReplacedLODModel)
+				{
+					delete ReplacedLODModel;
 				}
 			}
-		};
+			check(LODModelReplaceByReduction.IsEmpty());
+		}
 
-		SkeletalMesh->Modify();
-		SimplifyDependentLODS();
+		GWarn->EndSlowTask();
 	}
 }
 
diff --git a/Engine/Source/Editor/UnrealEd/Private/Factories/SkeletalMeshImport.cpp b/Engine/Source/Editor/UnrealEd/Private/Factories/SkeletalMeshImport.cpp
--- a/Engine/Source/Editor/UnrealEd/Private/Factories/SkeletalMeshImport.cpp
+++ b/Engine/Source/Editor/UnrealEd/Private/Factories/SkeletalMeshImport.cpp
@@ -462,33 +462,6 @@ ExistingSkelMeshData* SaveExistingSkelMeshData(USkeletalMesh* ExistingSkelMesh,
 	return ExistingMeshDataPtr;
 }
 
-void RegenerateDependentLODs(USkeletalMesh* SkelMesh, int32 LODIndex)
-{
-	check(SkelMesh);
-
-	IMeshReductionModule& ReductionModule = FModuleManager::Get().LoadModuleChecked<IMeshReductionModule>("MeshReductionInterface");
-	IMeshReduction* MeshReduction = ReductionModule.GetSkeletalMeshReductionInterface();
-	if (MeshReduction && MeshReduction->IsSupported())
-	{
-		FSkeletalMeshUpdateContext UpdateContext;
-		UpdateContext.SkeletalMesh = SkelMesh;
-		TArray<bool> DependentLODs;
-		DependentLODs.AddZeroed(SkelMesh->GetLODNum());
-		DependentLODs[LODIndex] = true;
-		for (int32 CurrentLODIndex = LODIndex + 1; CurrentLODIndex < DependentLODs.Num(); ++CurrentLODIndex)
-		{
-			FSkeletalMeshLODInfo& CurrentLODInfo = *(SkelMesh->GetLODInfo(CurrentLODIndex));
-			FSkeletalMeshOptimizationSettings& Settings = CurrentLODInfo.ReductionSettings;
-			if (CurrentLODInfo.bHasBeenSimplified && DependentLODs[Settings.BaseLOD])
-			{
-				DependentLODs[CurrentLODIndex] = true;
-				//Regenerate this LOD
-				FLODUtilities::SimplifySkeletalMeshLOD(UpdateContext, CurrentLODIndex, false);
-			}
-		}
-	}
-}
-
 void RestoreDependentLODs(ExistingSkelMeshData* MeshData, USkeletalMesh* SkeletalMesh)
 {
 	check(SkeletalMesh != nullptr);
diff --git a/Engine/Source/Runtime/Engine/Private/SkeletalMesh.cpp b/Engine/Source/Runtime/Engine/Private/SkeletalMesh.cpp
--- a/Engine/Source/Runtime/Engine/Private/SkeletalMesh.cpp
+++ b/Engine/Source/Runtime/Engine/Private/SkeletalMesh.cpp
@@ -325,6 +325,11 @@ FScopedSkeletalMeshPostEditChange::~FScopedSkeletalMeshPostEditChange()
 
 void FScopedSkeletalMeshPostEditChange::SetSkeletalMesh(USkeletalMesh* InSkeletalMesh)
 {
+	//Some parallel task may try to call post edit change, we must prevent it
+	if (!IsInGameThread())
+	{
+		return;
+	}
 	//We cannot set a different skeletal mesh, check that it was construct with null
 	check(SkeletalMesh == nullptr);
 	//We can only set a valid skeletal mesh
@@ -2417,12 +2422,15 @@ void USkeletalMesh::InitMorphTargetsAndRebuildRenderData()
 	// need to refresh the map
 	InitMorphTargets();
 
-	// reset all morphtarget for all components
-	for (TObjectIterator<USkeletalMeshComponent> It; It; ++It)
+	if (IsInGameThread())
 	{
-		if (It->SkeletalMesh == this)
+		// reset all morphtarget for all components
+		for (TObjectIterator<USkeletalMeshComponent> It; It; ++It)
 		{
-			It->RefreshMorphTargets();
+			if (It->SkeletalMesh == this)
+			{
+				It->RefreshMorphTargets();
+			}
 		}
 	}
 }
diff --git a/Engine/Source/Runtime/Engine/Public/SkeletalMeshReductionSettings.h b/Engine/Source/Runtime/Engine/Public/SkeletalMeshReductionSettings.h
--- a/Engine/Source/Runtime/Engine/Public/SkeletalMeshReductionSettings.h
+++ b/Engine/Source/Runtime/Engine/Public/SkeletalMeshReductionSettings.h
@@ -11,6 +11,7 @@
 #include "BoneContainer.h"
 #include "SkeletalMeshReductionSettings.generated.h"
 
+class FSkeletalMeshLODModel;
 
 /** Enum specifying the reduction type to use when simplifying skeletal meshes with internal tool */
 UENUM()
@@ -48,6 +49,8 @@ enum SkeletalMeshOptimizationImportance
 	SMOI_MAX UMETA(Hidden)
 };
 
+DECLARE_DELEGATE_OneParam(FOnDeleteLODModelOverride, FSkeletalMeshLODModel*);
+
 /**
 * FSkeletalMeshOptimizationSettings - The settings used to optimize a skeletal mesh LOD.
 */
@@ -145,6 +148,9 @@ struct FSkeletalMeshOptimizationSettings
 
 	UPROPERTY()
 	class UAnimSequence* BakePose_DEPRECATED;
+
+	//Transient mutable delegate. If the delegate is bound, the reduction will call it instead of deleting the replaced LODModel. It will be then the delegate owner responsible of the LODModel memory
+	mutable FOnDeleteLODModelOverride OnDeleteLODModelDelegate;
 #endif
 
 	FSkeletalMeshOptimizationSettings()
