Prevent saving the clothing when reducing inline skeletalmesh LOD model. We do not need the clothing when doing inline reducing
#rnx

diff --git a/Engine/Source/Developer/SkeletalMeshUtilitiesCommon/Private/LODUtilities.cpp b/Engine/Source/Developer/SkeletalMeshUtilitiesCommon/Private/LODUtilities.cpp
--- a/Engine/Source/Developer/SkeletalMeshUtilitiesCommon/Private/LODUtilities.cpp
+++ b/Engine/Source/Developer/SkeletalMeshUtilitiesCommon/Private/LODUtilities.cpp
@@ -1022,8 +1022,20 @@ void FLODUtilities::SimplifySkeletalMeshLOD( USkeletalMesh* SkeletalMesh, int32
 						MorphDeltasArray.Add(MorphDelta);
 					}
 				}
+				
 				//Copy the original SkeletalMesh LODModel
-				SkeletalMeshResource->OriginalReductionSourceMeshData[DesiredLOD]->SaveReductionData(SrcModel, BaseLODMorphTargetData, SkeletalMesh);
+				// Unbind clothing before saving the original data, we must not restore clothing to do inline reduction
+				{
+					TArray<ClothingAssetUtils::FClothingAssetMeshBinding> TemporaryRemoveClothingBindings;
+					FLODUtilities::UnbindClothingAndBackup(SkeletalMesh, TemporaryRemoveClothingBindings, DesiredLOD);
+
+					SkeletalMeshResource->OriginalReductionSourceMeshData[DesiredLOD]->SaveReductionData(SrcModel, BaseLODMorphTargetData, SkeletalMesh);
+
+					if (TemporaryRemoveClothingBindings.Num() > 0)
+					{
+						FLODUtilities::RestoreClothingFromBackup(SkeletalMesh, TemporaryRemoveClothingBindings, DesiredLOD);
+					}
+				}
 
 				if (DesiredLOD == 0)
 				{
@@ -2445,22 +2457,28 @@ void FLODUtilities::UnbindClothingAndBackup(USkeletalMesh* SkeletalMesh, TArray<
 		{
 			check(Binding.Asset);
 			Binding.Asset->UnbindFromSkeletalMesh(SkeletalMesh, Binding.LODIndex);
+			
 			//Use the UserSectionsData original section index, this will ensure we remap correctly the cloth if the reduction has change the number of sections
 			int32 OriginalDataSectionIndex = LODModel.Sections[Binding.SectionIndex].OriginalDataSectionIndex;
 			Binding.SectionIndex = OriginalDataSectionIndex;
+			
+			FSkelMeshSourceSectionUserData& SectionUserData = LODModel.UserSectionsData.FindChecked(OriginalDataSectionIndex);
+			SectionUserData.ClothingData.AssetGuid = FGuid();
+			SectionUserData.ClothingData.AssetLodIndex = INDEX_NONE;
+			SectionUserData.CorrespondClothAssetIndex = INDEX_NONE;
 		}
 	}
 }
 
-void FLODUtilities::RestoreClothingFromBackup(USkeletalMesh* SkeletalMesh, TArray<ClothingAssetUtils::FClothingAssetMeshBinding>& ClothingBindings, FOnPostRebindCloth OnPostRebindCloth)
+void FLODUtilities::RestoreClothingFromBackup(USkeletalMesh* SkeletalMesh, TArray<ClothingAssetUtils::FClothingAssetMeshBinding>& ClothingBindings)
 {
 	for (int32 LODIndex = 0; LODIndex < SkeletalMesh->GetImportedModel()->LODModels.Num(); ++LODIndex)
 	{
-		RestoreClothingFromBackup(SkeletalMesh, ClothingBindings, LODIndex, OnPostRebindCloth);
+		RestoreClothingFromBackup(SkeletalMesh, ClothingBindings, LODIndex);
 	}
 }
 
-void FLODUtilities::RestoreClothingFromBackup(USkeletalMesh* SkeletalMesh, TArray<ClothingAssetUtils::FClothingAssetMeshBinding>& ClothingBindings, const int32 LODIndex, FOnPostRebindCloth OnPostRebindCloth)
+void FLODUtilities::RestoreClothingFromBackup(USkeletalMesh* SkeletalMesh, TArray<ClothingAssetUtils::FClothingAssetMeshBinding>& ClothingBindings, const int32 LODIndex)
 {
 	if (!SkeletalMesh->GetImportedModel()->LODModels.IsValidIndex(LODIndex))
 	{
@@ -2480,7 +2498,10 @@ void FLODUtilities::RestoreClothingFromBackup(USkeletalMesh* SkeletalMesh, TArra
 				check(Binding.Asset);
 				if (Binding.Asset->BindToSkeletalMesh(SkeletalMesh, Binding.LODIndex, SectionIndex, Binding.AssetInternalLodIndex))
 				{
-					OnPostRebindCloth.ExecuteIfBound(Binding.LODIndex, LODModel.Sections[SectionIndex]);
+					//If successfull set back the section user data
+					FSkelMeshSourceSectionUserData& SectionUserData = LODModel.UserSectionsData.FindChecked(Binding.SectionIndex);
+					SectionUserData.CorrespondClothAssetIndex = LODModel.Sections[SectionIndex].CorrespondClothAssetIndex;
+					SectionUserData.ClothingData = LODModel.Sections[SectionIndex].ClothingData;
 				}
 			}
 			break;
@@ -2488,17 +2509,6 @@ void FLODUtilities::RestoreClothingFromBackup(USkeletalMesh* SkeletalMesh, TArra
 	}
 }
 
-void FLODUtilities::RestoreClothingFromBackup(USkeletalMesh* SkeletalMesh, TArray<ClothingAssetUtils::FClothingAssetMeshBinding>& ClothingBindings)
-{
-	FOnPostRebindCloth EmptyDelegate;
-	RestoreClothingFromBackup(SkeletalMesh, ClothingBindings, EmptyDelegate);
-}
-
-void FLODUtilities::RestoreClothingFromBackup(USkeletalMesh* SkeletalMesh, TArray<ClothingAssetUtils::FClothingAssetMeshBinding>& ClothingBindings, const int32 LODIndex)
-{
-	FOnPostRebindCloth EmptyDelegate;
-	RestoreClothingFromBackup(SkeletalMesh, ClothingBindings, LODIndex, EmptyDelegate);
-}
 
 
 #undef LOCTEXT_NAMESPACE // "LODUtilities"
\ No newline at end of file
diff --git a/Engine/Source/Developer/SkeletalMeshUtilitiesCommon/Public/LODUtilities.h b/Engine/Source/Developer/SkeletalMeshUtilitiesCommon/Public/LODUtilities.h
--- a/Engine/Source/Developer/SkeletalMeshUtilitiesCommon/Public/LODUtilities.h
+++ b/Engine/Source/Developer/SkeletalMeshUtilitiesCommon/Public/LODUtilities.h
@@ -27,7 +27,7 @@ struct FSkeletalMeshUpdateContext
 
 //////////////////////////////////////////////////////////////////////////
 // FLODUtilities
-DECLARE_DELEGATE_TwoParams(FOnPostRebindCloth, int32, const FSkelMeshSection&);
+
 
 class SKELETALMESHUTILITIESCOMMON_API FLODUtilities
 {
@@ -131,8 +131,6 @@ public:
 	
 	static void RestoreClothingFromBackup(USkeletalMesh* SkeletalMesh, TArray<ClothingAssetUtils::FClothingAssetMeshBinding>& ClothingBindings);
 	static void RestoreClothingFromBackup(USkeletalMesh* SkeletalMesh, TArray<ClothingAssetUtils::FClothingAssetMeshBinding>& ClothingBindings, const int32 LODIndex);
-	static void RestoreClothingFromBackup(USkeletalMesh* SkeletalMesh, TArray<ClothingAssetUtils::FClothingAssetMeshBinding>& ClothingBindings, FOnPostRebindCloth OnPostRebindCloth);
-	static void RestoreClothingFromBackup(USkeletalMesh* SkeletalMesh, TArray<ClothingAssetUtils::FClothingAssetMeshBinding>& ClothingBindings, const int32 LODIndex, FOnPostRebindCloth OnPostRebindCloth);
 	
 
 private:
diff --git a/Engine/Source/Editor/UnrealEd/Private/Fbx/FbxSkeletalMeshImport.cpp b/Engine/Source/Editor/UnrealEd/Private/Fbx/FbxSkeletalMeshImport.cpp
--- a/Engine/Source/Editor/UnrealEd/Private/Fbx/FbxSkeletalMeshImport.cpp
+++ b/Engine/Source/Editor/UnrealEd/Private/Fbx/FbxSkeletalMeshImport.cpp
@@ -2726,17 +2726,7 @@ USkeletalMesh* UnFbx::FFbxImporter::ReimportSkeletalMesh(USkeletalMesh* Mesh, UF
 			FSkeletalMeshModel* MeshResource = SkelMesh->GetImportedModel();
 			if (MeshResource && ClothingBindings.Num() > 0)
 			{
-				FLODUtilities::RestoreClothingFromBackup(SkelMesh, ClothingBindings, FOnPostRebindCloth::CreateLambda([&MeshResource](int32 LODIndex, const FSkelMeshSection& BindSection)
-				{
-					if (MeshResource->LODModels.IsValidIndex(LODIndex))
-					{
-						FSkeletalMeshLODModel& LODModel = MeshResource->LODModels[LODIndex];
-						//If successfull set back the section user data
-						FSkelMeshSourceSectionUserData& SectionUserData = LODModel.UserSectionsData.FindOrAdd(BindSection.OriginalDataSectionIndex);
-						SectionUserData.CorrespondClothAssetIndex = BindSection.CorrespondClothAssetIndex;
-						SectionUserData.ClothingData = BindSection.ClothingData;
-					}
-				}));
+				FLODUtilities::RestoreClothingFromBackup(SkelMesh, ClothingBindings);
 			}
 		};
 
