Copy up from Dev-Editor [REDACTED_HANDLE]

diff --git a/Engine/Plugins/MovieScene/MovieRenderPipeline/Source/MovieRenderPipelineEditor/Private/Widgets/SMoviePipelineQueuePanel.cpp b/Engine/Plugins/MovieScene/MovieRenderPipeline/Source/MovieRenderPipelineEditor/Private/Widgets/SMoviePipelineQueuePanel.cpp
new file mode 100644
--- /dev/null
+++ b/Engine/Plugins/MovieScene/MovieRenderPipeline/Source/MovieRenderPipelineEditor/Private/Widgets/SMoviePipelineQueuePanel.cpp
@@ -0,0 +1,295 @@
+// Copyright 1998-2019 Epic Games, Inc. All Rights Reserved.
+
+// Movie Pipeline Includes
+#include "Widgets/SMoviePipelineQueuePanel.h"
+#include "Widgets/MoviePipelineWidgetConstants.h"
+#include "SMoviePipelineQueueEditor.h"
+#include "SMoviePipelineConfigPanel.h"
+#include "MovieRenderPipelineSettings.h"
+#include "MovieRenderPipelineStyle.h"
+#include "Sections/MovieSceneCinematicShotSection.h"
+#include "MoviePipelineQueue.h"
+#include "MoviePipelineMasterConfig.h"
+#include "MoviePipelineQueueSubsystem.h"
+
+// Slate Includes
+#include "Widgets/SBoxPanel.h"
+#include "Widgets/Layout/SBorder.h"
+#include "Widgets/Text/STextBlock.h"
+#include "Widgets/Input/SButton.h"
+#include "Widgets/SWindow.h"
+#include "EditorStyleSet.h"
+#include "Framework/Application/SlateApplication.h"
+
+// Misc
+#include "Editor.h"
+
+#define LOCTEXT_NAMESPACE "SMoviePipelineQueuePanel"
+
+PRAGMA_DISABLE_OPTIMIZATION
+void SMoviePipelineQueuePanel::Construct(const FArguments& InArgs)
+{
+	// Allocate a transient preset automatically so they can start editing without having to create an asset.
+	// TransientPreset = AllocateTransientPreset();
+
+	// Copy the base preset into the transient preset if it was provided.
+	if (InArgs._BasePreset)
+	{
+		// TransientPreset->CopyFrom(InArgs._BasePreset);
+	}
+
+	// Create the child widgets that need to know about our pipeline
+	PipelineQueueEditorWidget = SNew(SMoviePipelineQueueEditor)
+		.OnEditConfigRequested(this, &SMoviePipelineQueuePanel::OnEditJobConfigRequested);
+	// .MoviePipeline(this, &SMoviePipelineConfigPanel::GetMoviePipeline);
+
+	ChildSlot
+	[
+		SNew(SVerticalBox)
+
+		// Create the toolbar for adding new items to the queue
+		+ SVerticalBox::Slot()
+		.Padding(FMargin(0.f, 1.0f))
+		.AutoHeight()
+		[
+			SNew(SBorder)
+			.BorderImage(FEditorStyle::GetBrush("DetailsView.CategoryTop"))
+			.BorderBackgroundColor(FLinearColor(.6, .6, .6, 1.0f))
+			[
+				SNew(SHorizontalBox)
+
+				// Add a Level Sequence to the queue 
+				+ SHorizontalBox::Slot()
+				.Padding(MoviePipeline::ButtonOffset)
+				.VAlign(VAlign_Fill)
+				.AutoWidth()
+				[
+					PipelineQueueEditorWidget->MakeAddSequenceJobButton()
+				]
+
+				// Remove a job (potentially already processed) from the the queue 
+				+ SHorizontalBox::Slot()
+				.Padding(MoviePipeline::ButtonOffset)
+				.VAlign(VAlign_Fill)
+				.AutoWidth()
+				[
+					PipelineQueueEditorWidget->RemoveSelectedJobButton()
+				]	
+			]
+		]
+	
+		// Main Queue Body
+		+ SVerticalBox::Slot()
+		.FillHeight(1.0f)
+		[
+			SNew(SHorizontalBox)
+			+ SHorizontalBox::Slot()
+			[
+				PipelineQueueEditorWidget.ToSharedRef()
+			]
+		]
+
+		// Footer Bar
+		+ SVerticalBox::Slot()
+		.AutoHeight()
+		[
+			SNew(SBorder)
+			.BorderImage(FEditorStyle::GetBrush("DetailsView.CategoryTop"))
+			.BorderBackgroundColor(FLinearColor(.6, .6, .6, 1.0f))
+			.Padding(FMargin(0, 2, 0, 2))
+			[
+				SNew(SHorizontalBox)
+
+				+ SHorizontalBox::Slot()
+				.VAlign(VAlign_Fill)
+				.HAlign(HAlign_Left)
+				.FillWidth(1.f)
+				[
+					SNullWidget::NullWidget
+				]
+
+				// Render Local in Process
+				+ SHorizontalBox::Slot()
+				.Padding(MoviePipeline::ButtonOffset)
+				.VAlign(VAlign_Fill)
+				.HAlign(HAlign_Right)
+				.AutoWidth()
+				[
+					SNew(SButton)
+					.ContentPadding(MoviePipeline::ButtonPadding)
+					.ButtonStyle(FMovieRenderPipelineStyle::Get(), "FlatButton.Success")
+					.IsEnabled(this, &SMoviePipelineQueuePanel::IsRenderLocalEnabled)
+					.OnClicked(this, &SMoviePipelineQueuePanel::OnRenderLocalRequested)
+					.Content()
+					[
+						SNew(STextBlock)
+						.TextStyle(FEditorStyle::Get(), "NormalText.Important")
+						.Text(LOCTEXT("RenderQueueLocal_Text", "Render (Local)"))
+						.ToolTipText(LOCTEXT("RenderQueueLocal_Tooltip", "Renders the current queue in the current process using Play in Editor."))
+						.Margin(FMargin(4, 0, 4, 0))
+					]
+				]
+
+				// Render Remotely (Separate Process or Farm)
+				+ SHorizontalBox::Slot()
+				.Padding(MoviePipeline::ButtonOffset)
+				.VAlign(VAlign_Fill)
+				.HAlign(HAlign_Right)
+				.AutoWidth()
+				[
+					SNew(SButton)
+					.ContentPadding(MoviePipeline::ButtonPadding)
+					.ButtonStyle(FMovieRenderPipelineStyle::Get(), "FlatButton.Success")
+					.IsEnabled(this, &SMoviePipelineQueuePanel::IsRenderRemoteEnabled)
+					.OnClicked(this, &SMoviePipelineQueuePanel::OnRenderRemoteRequested)
+					.Content()
+					[
+						SNew(STextBlock)
+						.TextStyle(FEditorStyle::Get(), "NormalText.Important")
+						.Text(LOCTEXT("RenderQueueRemote_Text", "Render (Remote)"))
+						.ToolTipText(LOCTEXT("RenderQueueRemote_Tooltip", "Renders the current queue in a separate process."))
+						.Margin(FMargin(4, 0, 4, 0))
+					]
+				]
+			]
+		]
+	];
+}
+
+PRAGMA_ENABLE_OPTIMIZATION
+
+FReply SMoviePipelineQueuePanel::OnRenderLocalRequested()
+{
+	UMoviePipelineQueueSubsystem* Subsystem = GEditor->GetEditorSubsystem<UMoviePipelineQueueSubsystem>();
+	check(Subsystem);
+
+	const UMovieRenderPipelineProjectSettings* ProjectSettings = GetDefault<UMovieRenderPipelineProjectSettings>();
+	check(ProjectSettings->DefaultLocalExecutor != nullptr);
+
+	Subsystem->RenderQueueWithExecutor(ProjectSettings->DefaultLocalExecutor);
+	return FReply::Handled();
+}
+
+bool SMoviePipelineQueuePanel::IsRenderLocalEnabled() const
+{
+	UMoviePipelineQueueSubsystem* Subsystem = GEditor->GetEditorSubsystem<UMoviePipelineQueueSubsystem>();
+	check(Subsystem);
+
+	const UMovieRenderPipelineProjectSettings* ProjectSettings = GetDefault<UMovieRenderPipelineProjectSettings>();
+	return (!Subsystem->IsRendering()) && (ProjectSettings->DefaultLocalExecutor != nullptr);
+}
+
+FReply SMoviePipelineQueuePanel::OnRenderRemoteRequested()
+{
+	UMoviePipelineQueueSubsystem* Subsystem = GEditor->GetEditorSubsystem<UMoviePipelineQueueSubsystem>();
+	check(Subsystem);
+
+	const UMovieRenderPipelineProjectSettings* ProjectSettings = GetDefault<UMovieRenderPipelineProjectSettings>();
+	check(ProjectSettings->DefaultRemoteExecutor != nullptr);
+
+	Subsystem->RenderQueueWithExecutor(ProjectSettings->DefaultRemoteExecutor);
+	return FReply::Handled();
+}
+
+bool SMoviePipelineQueuePanel::IsRenderRemoteEnabled() const
+{
+	UMoviePipelineQueueSubsystem* Subsystem = GEditor->GetEditorSubsystem<UMoviePipelineQueueSubsystem>();
+	check(Subsystem);
+
+	const UMovieRenderPipelineProjectSettings* ProjectSettings = GetDefault<UMovieRenderPipelineProjectSettings>();
+	return (!Subsystem->IsRendering()) && (ProjectSettings->DefaultRemoteExecutor != nullptr);
+}
+
+/*UMoviePipelineConfigBase* SMoviePipelineConfigPanel::AllocateTransientPreset()
+{
+	static const TCHAR* PackageName = TEXT("/Temp/MoviePipeline/PendingPipeline");
+
+	// Return a cached transient if it exists
+	UMoviePipelineConfigBase* ExistingPreset = FindObject<UMoviePipelineConfigBase>(nullptr, TEXT("/Temp/MoviePipeline/PendingPipeline.PendingPipeline"));
+	if (ExistingPreset)
+	{
+		return ExistingPreset;
+	}
+
+	static FName DesiredName = "PendingMoviePipeline";
+	
+	UPackage* NewPackage = CreatePackage(nullptr, PackageName);
+	NewPackage->SetFlags(RF_Transient);
+	NewPackage->AddToRoot();
+
+	UMoviePipelineConfigBase* NewPreset = NewObject<UMoviePipelineConfigBase>(NewPackage, DesiredName, RF_Transient | RF_Transactional | RF_Standalone);
+
+	return NewPreset;
+}*/
+
+void SMoviePipelineQueuePanel::AddReferencedObjects(FReferenceCollector& Collector)
+{
+	// Collector.AddReferencedObject(TransientPreset);
+	// Collector.AddReferencedObject(SuppliedLevelSequence);
+	// Collector.AddReferencedObject(RecordingLevelSequence);
+}
+
+void SMoviePipelineQueuePanel::OnEditJobConfigRequested(TWeakObjectPtr<UMoviePipelineExecutorJob> InJob, TWeakObjectPtr<UMovieSceneCinematicShotSection> InShot)
+{
+	TSharedRef<SWindow> EditorWindow =
+		SNew(SWindow)
+		.ClientSize(FVector2D(700, 600));
+
+	TSharedRef<SMoviePipelineConfigPanel> ConfigEditorPanel =
+		SNew(SMoviePipelineConfigPanel, UMoviePipelineMasterConfig::StaticClass())
+		.Job(InJob)
+		.OnConfigurationModified(this, &SMoviePipelineQueuePanel::OnConfigUpdatedForJob)
+		.OnConfigurationSetToPreset(this, &SMoviePipelineQueuePanel::OnConfigUpdatedForJobToPreset)
+		.BasePreset(InJob->GetPresetOrigin())
+		.BaseConfig(InJob->GetConfiguration());
+
+	EditorWindow->SetContent(ConfigEditorPanel);
+
+
+	TSharedPtr<SWindow> ParentWindow = FSlateApplication::Get().FindWidgetWindow(AsShared());
+	if (ParentWindow.IsValid())
+	{
+		FSlateApplication::Get().AddWindowAsNativeChild(EditorWindow, ParentWindow.ToSharedRef());
+	}
+
+	WeakEditorWindow = EditorWindow;
+}
+
+void SMoviePipelineQueuePanel::OnConfigWindowClosed()
+{
+	UE_LOG(LogTemp, Warning, TEXT("Config updated!"));
+	if (WeakEditorWindow.IsValid())
+	{
+		WeakEditorWindow.Pin()->RequestDestroyWindow();
+	}
+}
+
+void SMoviePipelineQueuePanel::OnConfigUpdatedForJob(TWeakObjectPtr<UMoviePipelineExecutorJob> InJob, UMoviePipelineConfigBase* InConfig)
+{
+	if (InJob.IsValid())
+	{
+		UMoviePipelineMasterConfig* MasterConfig = Cast<UMoviePipelineMasterConfig>(InConfig);
+		if (MasterConfig)
+		{
+			InJob->SetConfiguration(MasterConfig);
+		}
+	}
+
+	OnConfigWindowClosed();
+}
+
+void SMoviePipelineQueuePanel::OnConfigUpdatedForJobToPreset(TWeakObjectPtr<UMoviePipelineExecutorJob> InJob, UMoviePipelineConfigBase* InConfig)
+{
+	if (InJob.IsValid())
+	{
+		UMoviePipelineMasterConfig* MasterConfig = Cast<UMoviePipelineMasterConfig>(InConfig);
+		if (MasterConfig)
+		{
+			InJob->SetPresetOrigin(MasterConfig);
+		}
+	}
+
+	OnConfigWindowClosed();
+}
+
+#undef LOCTEXT_NAMESPACE // SMoviePipelineQueuePanel
\ No newline at end of file
diff --git a/Engine/Plugins/MovieScene/MovieRenderPipeline/Source/MovieRenderPipelineEditor/Private/Widgets/SMoviePipelineQueuePanel.h b/Engine/Plugins/MovieScene/MovieRenderPipeline/Source/MovieRenderPipelineEditor/Private/Widgets/SMoviePipelineQueuePanel.h
new file mode 100644
--- /dev/null
+++ b/Engine/Plugins/MovieScene/MovieRenderPipeline/Source/MovieRenderPipelineEditor/Private/Widgets/SMoviePipelineQueuePanel.h
@@ -0,0 +1,62 @@
+// Copyright 1998-2019 Epic Games, Inc. All Rights Reserved.
+
+#pragma once
+
+#include "Widgets/DeclarativeSyntaxSupport.h"
+#include "Widgets/SCompoundWidget.h"
+#include "UObject/GCObject.h"
+
+class UMoviePipelineConfigBase;
+class SMoviePipelineQueueEditor;
+class SWindow;
+class UMoviePipelineExecutorJob;
+class UMovieSceneCinematicShotSection;
+
+/**
+ * Outermost widget that is used for setting up a new movie render pipeline queue. Operates on a transient object that is internally owned and maintained 
+ */
+class SMoviePipelineQueuePanel : public SCompoundWidget, public FGCObject
+{
+public:
+	SLATE_BEGIN_ARGS(SMoviePipelineQueuePanel)
+		: _BasePreset(nullptr)
+
+		{}
+
+		/*~ All following arguments are mutually-exclusive */
+		/*-------------------------------------------------*/
+		/** A preset asset to base the pipeline off */
+		SLATE_ARGUMENT(UMoviePipelineConfigBase*, BasePreset)
+
+	SLATE_END_ARGS()
+
+	void Construct(const FArguments& InArgs);
+
+private:
+	// FGCObject Interface
+	virtual void AddReferencedObjects(FReferenceCollector& Collector) override;
+	// ~FGCObject Interface
+
+	FReply OnRenderLocalRequested();
+	bool IsRenderLocalEnabled() const;
+	FReply OnRenderRemoteRequested();
+	bool IsRenderRemoteEnabled() const;
+
+	void OnEditJobConfigRequested(TWeakObjectPtr<UMoviePipelineExecutorJob> InJob, TWeakObjectPtr<UMovieSceneCinematicShotSection> InShot);
+	void OnConfigUpdatedForJob(TWeakObjectPtr<UMoviePipelineExecutorJob> InJob, UMoviePipelineConfigBase* InConfig);
+	void OnConfigUpdatedForJobToPreset(TWeakObjectPtr<UMoviePipelineExecutorJob> InJob, UMoviePipelineConfigBase* InConfig);
+	void OnConfigWindowClosed();
+
+private:
+	/** Allocates a transient preset so that the user can use the pipeline without saving it to an asset first. */
+	//UMoviePipelineConfigBase* AllocateTransientPreset();
+
+
+private:
+	/** The main movie pipeline queue editor widget */
+	TSharedPtr<SMoviePipelineQueueEditor> PipelineQueueEditorWidget;
+
+	TWeakPtr<SWindow> WeakEditorWindow;
+	/** The transient preset that we use - kept alive by AddReferencedObjects */
+	// UMoviePipelineConfigBase* TransientPreset;
+};
\ No newline at end of file
diff --git a/Engine/Plugins/MovieScene/MovieRenderPipeline/Source/MovieRenderPipelineEditor/Public/MoviePipelineEditorBlueprintLibrary.h b/Engine/Plugins/MovieScene/MovieRenderPipeline/Source/MovieRenderPipelineEditor/Public/MoviePipelineEditorBlueprintLibrary.h
new file mode 100644
--- /dev/null
+++ b/Engine/Plugins/MovieScene/MovieRenderPipeline/Source/MovieRenderPipelineEditor/Public/MoviePipelineEditorBlueprintLibrary.h
@@ -0,0 +1,25 @@
+// Copyright 1998-2019 Epic Games, Inc. All Rights Reserved.
+
+#pragma once
+
+#include "Kismet/BlueprintFunctionLibrary.h"
+#include "UObject/ObjectMacros.h"
+#include "UObject/TextProperty.h"
+#include "Internationalization/Text.h"
+#include "Containers/UnrealString.h"
+#include "MoviePipelineMasterConfig.h"
+
+#include "MoviePipelineEditorBlueprintLibrary.generated.h"
+
+// Forward Declare
+class UMoviePipelineMasterConfig;
+
+UCLASS()
+class MOVIERENDERPIPELINEEDITOR_API UMoviePipelineEditorBlueprintLibrary : public UBlueprintFunctionLibrary
+{
+	GENERATED_BODY()
+	
+public:
+	UFUNCTION(BlueprintCallable, Category = "Movie Render Pipeline")
+	static bool ExportConfigToAsset(const UMoviePipelineMasterConfig* InConfig, const FString& InPackagePath, const FString& InFileName, const bool bInSaveAsset, UMoviePipelineMasterConfig*& OutAsset, FText& OutErrorReason);
+};
\ No newline at end of file
diff --git a/Engine/Plugins/MovieScene/MovieRenderPipeline/Source/MovieRenderPipelineEditor/Public/MoviePipelineNewProcessExecutor.h b/Engine/Plugins/MovieScene/MovieRenderPipeline/Source/MovieRenderPipelineEditor/Public/MoviePipelineNewProcessExecutor.h
new file mode 100644
--- /dev/null
+++ b/Engine/Plugins/MovieScene/MovieRenderPipeline/Source/MovieRenderPipelineEditor/Public/MoviePipelineNewProcessExecutor.h
@@ -0,0 +1,36 @@
+// Copyright 1998-2019 Epic Games, Inc. All Rights Reserved.
+#pragma once
+
+#include "MoviePipelineLinearExecutor.h"
+#include "MoviePipelineNewProcessExecutor.generated.h"
+
+/**
+* This is the implementation responsible for executing the rendering of
+* multiple movie pipelines on the local machine in an external process.
+* This simply handles launching and managing the external processes and 
+* acts as a proxy to them where possible. This internally uses the
+* UMoviePipelineInProcessExecutor on the launched instances.
+*/
+UCLASS(Blueprintable)
+class MOVIERENDERPIPELINEEDITOR_API UMoviePipelineNewProcessExecutor : public UMoviePipelineLinearExecutorBase
+{
+	GENERATED_BODY()
+	
+public:
+	UMoviePipelineNewProcessExecutor()
+		: UMoviePipelineLinearExecutorBase()
+	{
+	}
+
+	/** If enabled the editor will close itself when a new process is started. This can be used to gain some performance. */
+	UPROPERTY(EditAnywhere, BlueprintReadWrite, Category = "Movie Pipeline")
+	bool bCloseEditor;
+
+	/** A list of additional command line arguments to be appended to the new process startup. Can be useful if your game requires certain arguments to start such as disabling log-in screens. */
+	UPROPERTY(EditAnywhere, BlueprintReadWrite, Category = "Movie Pipeline")
+	FString AdditionalCommandLineArguments;
+
+	/** A list of command line arguments which are inherited from the currently running Editor instance that will be automatically appended to the new process. */
+	UPROPERTY(VisibleAnywhere, BlueprintReadOnly, Category = "Movie Pipeline")
+	FString InheritedCommandLineArguments;
+};
\ No newline at end of file
diff --git a/Engine/Plugins/MovieScene/MovieRenderPipeline/Source/MovieRenderPipelineEditor/Public/MoviePipelinePIEExecutor.h b/Engine/Plugins/MovieScene/MovieRenderPipeline/Source/MovieRenderPipelineEditor/Public/MoviePipelinePIEExecutor.h
new file mode 100644
--- /dev/null
+++ b/Engine/Plugins/MovieScene/MovieRenderPipeline/Source/MovieRenderPipelineEditor/Public/MoviePipelinePIEExecutor.h
@@ -0,0 +1,41 @@
+// Copyright 1998-2019 Epic Games, Inc. All Rights Reserved.
+#pragma once
+
+#include "MoviePipelineLinearExecutor.h"
+#include "MoviePipelinePIEExecutor.generated.h"
+
+class UMoviePipeline;
+
+/**
+* This is the implementation responsible for executing the rendering of
+* multiple movie pipelines in the currently running Editor process. This
+* involves launching a Play in Editor session for each Movie Pipeline to
+* process.
+*/
+UCLASS(Blueprintable)
+class MOVIERENDERPIPELINEEDITOR_API UMoviePipelinePIEExecutor : public UMoviePipelineLinearExecutorBase
+{
+	GENERATED_BODY()
+	
+public:
+	UMoviePipelinePIEExecutor()
+		: UMoviePipelineLinearExecutorBase()
+	{
+	}
+
+protected:
+	virtual void Start(const UMoviePipelineExecutorJob* InJob) override;
+
+private:
+	/** Called when PIE finishes booting up and it is safe for us to spawn an object into that world. */
+	void OnPIEStartupFinished(bool);
+	/** Called before PIE tears down the world during shutdown. Used to detect cancel-via-escape/stop PIE. */
+	void OnPIEEnded(bool);
+	/** Called when the instance of the pipeline in the PIE world has finished. */
+	void OnPIEMoviePipelineFinished(UMoviePipeline* InMoviePipeline);
+	/** Called a short period of time after OnPIEMoviePipelineFinished to allow Editor the time to fully close PIE before we make a new request. */
+	void DelayedFinishNotification();
+
+	/** Gets the title for the Window. Used to show progress in the title. */
+	FText GetWindowTitle(const int32 InConfigIndex, const int32 InNumConfigs) const;
+};
\ No newline at end of file
diff --git a/Engine/Plugins/MovieScene/MovieRenderPipeline/Source/MovieRenderPipelineRenderPasses/Private/MovieRenderPipelineRenderPassesModule.cpp b/Engine/Plugins/MovieScene/MovieRenderPipeline/Source/MovieRenderPipelineRenderPasses/Private/MovieRenderPipelineRenderPassesModule.cpp
new file mode 100644
--- /dev/null
+++ b/Engine/Plugins/MovieScene/MovieRenderPipeline/Source/MovieRenderPipelineRenderPasses/Private/MovieRenderPipelineRenderPassesModule.cpp
@@ -0,0 +1,37 @@
+// Copyright 1998-2019 Epic Games, Inc. All Rights Reserved.
+
+#include "Modules/ModuleManager.h"
+#include "MovieRenderPipelineDataTypes.h"
+#include "MovieRenderPipelineCoreModule.h"
+#include "MoviePipelineDeferredPasses.h"
+
+class FMovieRenderPipelineRenderPassesModule : public IModuleInterface
+{
+
+	static TSharedRef<MoviePipeline::FMoviePipelineEnginePass> CreateDeferredEnginePass()
+	{
+		return MakeShared<MoviePipeline::FDeferredRenderEnginePass>();
+	}
+
+	virtual void StartupModule() override
+	{
+		FMovieRenderPipelineCoreModule& CoreModule = FModuleManager::Get().LoadModuleChecked<FMovieRenderPipelineCoreModule>("MovieRenderPipelineCore");
+		DeferredEnginePassHandle = CoreModule.RegisterEngineRenderPass(FOnCreateEngineRenderPass::CreateStatic(&FMovieRenderPipelineRenderPassesModule::CreateDeferredEnginePass));
+	}
+
+	virtual void ShutdownModule() override
+	{
+		FMovieRenderPipelineCoreModule* CoreModule = (FMovieRenderPipelineCoreModule*)FModuleManager::Get().LoadModule("MovieRenderPipelineCore");
+		if (CoreModule)
+		{
+			CoreModule->UnregisterEngineRenderPass(DeferredEnginePassHandle);
+		}
+
+	}
+private:
+
+	FDelegateHandle DeferredEnginePassHandle;
+};
+
+IMPLEMENT_MODULE(FMovieRenderPipelineRenderPassesModule, MovieRenderPipelineRenderPasses);
+
diff --git a/Engine/Plugins/MovieScene/MovieRenderPipeline/Source/MovieRenderPipelineRenderPasses/Public/MoviePipelineImageSequenceOutput.h b/Engine/Plugins/MovieScene/MovieRenderPipeline/Source/MovieRenderPipelineRenderPasses/Public/MoviePipelineImageSequenceOutput.h
new file mode 100644
--- /dev/null
+++ b/Engine/Plugins/MovieScene/MovieRenderPipeline/Source/MovieRenderPipelineRenderPasses/Public/MoviePipelineImageSequenceOutput.h
@@ -0,0 +1,120 @@
+// Copyright 1998-2019 Epic Games, Inc. All Rights Reserved.
+#pragma once
+
+#include "MoviePipelineOutputBase.h"
+#include "Async/Future.h"
+#include "MoviePipelineImageSequenceOutput.generated.h"
+
+// Forward Declare
+class IImageWriteQueue;
+
+UCLASS(Blueprintable)
+class MOVIERENDERPIPELINERENDERPASSES_API UMoviePipelineImageSequenceOutputBase : public UMoviePipelineOutputBase
+{
+	GENERATED_BODY()
+public:
+	UMoviePipelineImageSequenceOutputBase();
+
+	virtual void OnRecieveImageDataImpl(FMoviePipelineMergerOutputFrame* InMergedOutputFrame) override;
+
+protected:
+	// UMovieRenderPipelineOutputContainer interface
+	virtual void BeginFinalizeImpl() override;
+	virtual bool HasFinishedProcessingImpl() override;
+	// ~UMovieRenderPipelineOutputContainer interface
+
+protected:
+	/** The format of the image to write out */
+	EImageFormat OutputFormat;
+
+private:
+	/** A pointer to the image write queue used for asynchronously writing images */
+	IImageWriteQueue* ImageWriteQueue;
+
+	/** A fence to keep track of when the Image Write queue has fully flushed. */
+	TFuture<void> FinalizeFence;
+};
+
+UCLASS()
+class MOVIERENDERPIPELINERENDERPASSES_API UMoviePipelineImageSequenceOutput_BMP : public UMoviePipelineImageSequenceOutputBase
+{
+	GENERATED_BODY()
+public:
+#if WITH_EDITOR
+	virtual FText GetDisplayText() const override { return NSLOCTEXT("MovieRenderPipeline", "ImgSequenceBMPSettingDisplayName", "Image Sequence (.bmp [8bpp])"); }
+#endif
+public:
+	UMoviePipelineImageSequenceOutput_BMP()
+	{
+		OutputFormat = EImageFormat::BMP;
+	}
+};
+
+UCLASS()
+class MOVIERENDERPIPELINERENDERPASSES_API UMoviePipelineImageSequenceOutput_PNG : public UMoviePipelineImageSequenceOutputBase
+{
+	GENERATED_BODY()
+public:
+#if WITH_EDITOR
+	virtual FText GetDisplayText() const override { return NSLOCTEXT("MovieRenderPipeline", "ImgSequencePNGSettingDisplayName", "Image Sequence (.png [8bpp])"); }
+#endif
+public:
+	UMoviePipelineImageSequenceOutput_PNG()
+	{
+		OutputFormat = EImageFormat::PNG;
+	}
+
+	virtual bool IsAlphaSupported() const override { return bOutputAlpha; }
+
+public:
+	/**
+	* Should we accumulate the alpha channel and write it into the resulting image? This requires r.PostProcessing.PropagateAlpha
+	* to be set to 1 or 2 (see "Enable Alpha Channel Support in Post Processing" under Project Settings > Rendering). This adds
+	* ~30% cost to the accumulation so you should not enable it unless necessary. You must delete both the sky and fog to ensure
+	* that they do not make all pixels opaque.
+	*/
+	UPROPERTY(EditAnywhere, BlueprintReadWrite, Category = "PNG")
+	bool bOutputAlpha;
+};
+
+UCLASS()
+class MOVIERENDERPIPELINERENDERPASSES_API UMoviePipelineImageSequenceOutput_JPG : public UMoviePipelineImageSequenceOutputBase
+{
+	GENERATED_BODY()
+public:
+#if WITH_EDITOR
+	virtual FText GetDisplayText() const override { return NSLOCTEXT("MovieRenderPipeline", "ImgSequenceJPGSettingDisplayName", "Image Sequence (.jpg [8bpp])"); }
+#endif
+public:
+	UMoviePipelineImageSequenceOutput_JPG()
+	{
+		OutputFormat = EImageFormat::JPEG;
+	}
+};
+
+UCLASS()
+class MOVIERENDERPIPELINERENDERPASSES_API UMoviePipelineImageSequenceOutput_EXR : public UMoviePipelineImageSequenceOutputBase
+{
+	GENERATED_BODY()
+public:
+#if WITH_EDITOR
+	virtual FText GetDisplayText() const override { return NSLOCTEXT("MovieRenderPipeline", "ImgSequenceEXRSettingDisplayName", "Image Sequence (.exr [32bpp])"); }
+#endif
+public:
+	UMoviePipelineImageSequenceOutput_EXR()
+	{
+		OutputFormat = EImageFormat::EXR;
+	}
+
+	virtual bool IsAlphaSupported() const override { return bOutputAlpha; }
+
+public:
+	/**
+	* Should we accumulate the alpha channel and write it into the resulting image? This requires r.PostProcessing.PropagateAlpha
+	* to be set to 1 or 2 (see "Enable Alpha Channel Support in Post Processing" under Project Settings > Rendering). This adds
+	* ~30% cost to the accumulation so you should not enable it unless necessary. You must delete both the sky and fog to ensure
+	* that they do not make all pixels opaque.
+	*/
+	UPROPERTY(EditAnywhere, BlueprintReadWrite, Category = "EXR")
+	bool bOutputAlpha;
+};
\ No newline at end of file
diff --git a/Engine/Plugins/MovieScene/MovieRenderPipeline/Source/MovieRenderPipelineSettings/Public/MoviePipelineBurnInSetting.h b/Engine/Plugins/MovieScene/MovieRenderPipeline/Source/MovieRenderPipelineSettings/Public/MoviePipelineBurnInSetting.h
new file mode 100644
--- /dev/null
+++ b/Engine/Plugins/MovieScene/MovieRenderPipeline/Source/MovieRenderPipelineSettings/Public/MoviePipelineBurnInSetting.h
@@ -0,0 +1,24 @@
+// Copyright 1998-2019 Epic Games, Inc. All Rights Reserved.
+#pragma once
+
+#include "MoviePipelineSetting.h"
+#include "LevelSequenceActor.h"
+#include "MoviePipelineBurnInSetting.generated.h"
+
+UCLASS(Blueprintable)
+class UMoviePipelineBurnInSetting : public UMoviePipelineSetting
+{
+	GENERATED_BODY()
+public:
+#if WITH_EDITOR
+	virtual FText GetDisplayText() const override { return NSLOCTEXT("MovieRenderPipeline", "BurnInSettingDisplayName", "Burn In"); }
+#endif
+	virtual bool IsValidOnShots() const override { return false; }
+	virtual bool IsValidOnMaster() const override { return true; }
+public:
+	UPROPERTY(EditAnywhere, BlueprintReadWrite, meta=(MetaClass="LevelSequenceBurnIn"), Category = "Movie Pipeline")
+	FSoftClassPath BurnInClass;
+
+	UPROPERTY(Instanced, EditAnywhere, BlueprintReadWrite, Category="Movie Pipeline")
+	ULevelSequenceBurnInInitSettings* Settings;
+};
\ No newline at end of file
diff --git a/Engine/Plugins/MovieScene/MovieRenderPipeline/Source/MovieRenderPipelineSettings/Public/MoviePipelineGameOverrideSetting.h b/Engine/Plugins/MovieScene/MovieRenderPipeline/Source/MovieRenderPipelineSettings/Public/MoviePipelineGameOverrideSetting.h
new file mode 100644
--- /dev/null
+++ b/Engine/Plugins/MovieScene/MovieRenderPipeline/Source/MovieRenderPipelineSettings/Public/MoviePipelineGameOverrideSetting.h
@@ -0,0 +1,42 @@
+// Copyright 1998-2019 Epic Games, Inc. All Rights Reserved.
+#pragma once
+
+#include "MoviePipelineSetting.h"
+#include "GameFramework/GameModeBase.h"
+#include "Templates/SubclassOf.h"
+#include "MoviePipelineGameOverrideSetting.generated.h"
+
+// Forward Declares
+class AGameModeBase;
+
+UCLASS(Blueprintable)
+class MOVIERENDERPIPELINESETTINGS_API UMoviePipelineGameOverrideSetting : public UMoviePipelineSetting
+{
+	GENERATED_BODY()
+public:
+	UMoviePipelineGameOverrideSetting()
+		: bDisableCinematicMode(false)
+		, bShowPlayer(false)
+	{
+	}
+
+public:
+#if WITH_EDITOR
+	virtual FText GetDisplayText() const override { return NSLOCTEXT("MovieRenderPipeline", "GameOverrideSettingDisplayName", "Game Overrides"); }
+#endif
+	virtual bool IsValidOnShots() const override { return false; }
+	virtual bool IsValidOnMaster() const override { return true; }
+
+public:
+	/** Optional Game Mode to override the map's default game mode with. This can be useful if the game's normal mode displays UI elements or loading screens that you don't want captured. */
+	UPROPERTY(EditAnywhere, BlueprintReadWrite, Category = "Movie Pipeline")
+	TSubclassOf<AGameModeBase> GameModeOverride;
+		
+	/** Cinematic Mode makes various systems to flush async loading to ensure first frame accuracy, which also disables Player Controller input, and the HUD spawned by AGameModeBase::HUDClass */
+	UPROPERTY(EditAnywhere, BlueprintReadWrite, Category = "Movie Pipeline")
+	bool bDisableCinematicMode;
+	
+	/** Should the Player Pawn spawned by the maps gamemode be visible? */
+	UPROPERTY(EditAnywhere, BlueprintReadWrite, Category = "Movie Pipeline")
+	bool bShowPlayer;
+};
\ No newline at end of file
diff --git a/Engine/Plugins/Runtime/nDisplay/Source/DisplayClusterEditor/Private/DisplayClusterEditorEngine.cpp b/Engine/Plugins/Runtime/nDisplay/Source/DisplayClusterEditor/Private/DisplayClusterEditorEngine.cpp
--- a/Engine/Plugins/Runtime/nDisplay/Source/DisplayClusterEditor/Private/DisplayClusterEditorEngine.cpp
+++ b/Engine/Plugins/Runtime/nDisplay/Source/DisplayClusterEditor/Private/DisplayClusterEditorEngine.cpp
@@ -66,18 +66,20 @@ ADisplayClusterRootActor* UDisplayClusterEditorEngine::FindDisplayClusterRootAct
 	return nullptr;
 }
 
-void UDisplayClusterEditorEngine::PlayInEditor(UWorld* InWorld, bool bInSimulateInEditor, FPlayInEditorOverrides Overrides)
+void UDisplayClusterEditorEngine::StartPlayInEditorSession(FRequestPlaySessionParams& InRequestParams)
 {
-	UE_LOG(LogDisplayClusterEditorEngine, VeryVerbose, TEXT("UDisplayClusterEditorEngine::PlayInEditor"));
+	UE_LOG(LogDisplayClusterEditorEngine, VeryVerbose, TEXT("UDisplayClusterEditorEngine::StartPlayInEditorSession"));
+
+	UWorld* EditorWorldPreDup = GetEditorWorldContext().World();
 
 	if (DisplayClusterModule)
 	{
 		// Find nDisplay root actor
-		ADisplayClusterRootActor* RootActor = FindDisplayClusterRootActor(InWorld);
+		ADisplayClusterRootActor* RootActor = FindDisplayClusterRootActor(EditorWorldPreDup);
 		if (!RootActor)
 		{
 			// Also search inside streamed levels
-			const TArray<ULevelStreaming*>& StreamingLevels = InWorld->GetStreamingLevels();
+			const TArray<ULevelStreaming*>& StreamingLevels = EditorWorld->GetStreamingLevels();
 			for (ULevelStreaming* StreamingLevel : StreamingLevels)
 			{
 				switch (StreamingLevel->GetCurrentState())
@@ -113,11 +115,11 @@ void UDisplayClusterEditorEngine::PlayInEditor(UWorld* InWorld, bool bInSimulate
 				return;
 			}
 
-			DisplayClusterModule->StartScene(InWorld);
+			DisplayClusterModule->StartScene(EditorWorldPreDup);
 		}
 	}
 
-	Super::PlayInEditor(InWorld, bInSimulateInEditor, Overrides);
+	Super::StartPlayInEditorSession(InRequestParams);
 }
 
 void UDisplayClusterEditorEngine::Tick(float DeltaSeconds, bool bIdleMode)
diff --git a/Engine/Plugins/Tests/EditorTests/Source/EditorTests/Private/UnrealEd/EditorBuildPromotionTests.cpp b/Engine/Plugins/Tests/EditorTests/Source/EditorTests/Private/UnrealEd/EditorBuildPromotionTests.cpp
--- a/Engine/Plugins/Tests/EditorTests/Source/EditorTests/Private/UnrealEd/EditorBuildPromotionTests.cpp
+++ b/Engine/Plugins/Tests/EditorTests/Source/EditorTests/Private/UnrealEd/EditorBuildPromotionTests.cpp
@@ -395,8 +395,15 @@ namespace EditorBuildPromotionTestUtils
 	{
 		FLevelEditorModule& LevelEditorModule = FModuleManager::Get().GetModuleChecked<FLevelEditorModule>(TEXT("LevelEditor"));
 		TSharedPtr<class IAssetViewport> ActiveLevelViewport = LevelEditorModule.GetFirstActiveViewport();
+		
+		FRequestPlaySessionParams SessionParams;
+		SessionParams.DestinationSlateViewport = ActiveLevelViewport;
+		if (bSimulateInEditor)
+		{
+			SessionParams.WorldType = EPlaySessionWorldType::SimulateInEditor;
+		}
 
-		GUnrealEd->RequestPlaySession(false, ActiveLevelViewport, bSimulateInEditor, NULL, NULL, -1, false);
+		GUnrealEd->RequestPlaySession(SessionParams);
 	}
 
 	/**
diff --git a/Engine/Plugins/VirtualProduction/Takes/Source/TakeRecorder/Private/Widgets/SLevelSequenceTakeEditor.cpp b/Engine/Plugins/VirtualProduction/Takes/Source/TakeRecorder/Private/Widgets/SLevelSequenceTakeEditor.cpp
--- a/Engine/Plugins/VirtualProduction/Takes/Source/TakeRecorder/Private/Widgets/SLevelSequenceTakeEditor.cpp
+++ b/Engine/Plugins/VirtualProduction/Takes/Source/TakeRecorder/Private/Widgets/SLevelSequenceTakeEditor.cpp
@@ -778,11 +778,23 @@ void SLevelSequenceTakeEditor::UpdateDetails()
 	FDetailsViewArgs DetailsViewArgs(false, false, false, FDetailsViewArgs::HideNameArea, true);
 	DetailsViewArgs.bShowScrollBar = false;
 
+	TSortedMap<const UClass*, TArray<UObject*>> ClassToSources;
+
+	// External settings objects should be displayed after selected sources, 
+	// so add them to the sorted map first, since the details views will be displayed in reverse order
+	for (TWeakObjectPtr<> WeakExternalObj : ExternalSettingsObjects)
+	{
+		UObject* Object = WeakExternalObj.Get();
+		if (Object)
+		{
+			ClassToSources.FindOrAdd(Object->GetClass()).Add(Object);
+		}
+	}
+
 	TArray<UTakeRecorderSource*> SelectedSources;
 	SourcesWidget->GetSelectedSources(SelectedSources);
 
 	// Create 1 details panel per source class type
-	TSortedMap<const UClass*, TArray<UObject*>> ClassToSources;
 	for (UTakeRecorderSource* Source : SelectedSources)
 	{
 		ClassToSources.FindOrAdd(Source->GetClass()).Add(Source);
@@ -796,17 +808,21 @@ void SLevelSequenceTakeEditor::UpdateDetails()
 		}
 	}
 
-	for (TWeakObjectPtr<> WeakExternalObj : ExternalSettingsObjects)
+	TArray<FObjectKey> PreviousClasses;
+	ClassToDetailsView.GenerateKeyArray(PreviousClasses);
+
+	// Clear all existing details views if there are external settings objects, so that they can be displayed last
+	if (ExternalSettingsObjects.Num())
 	{
-		UObject* Object = WeakExternalObj.Get();
-		if (Object)
+		for (FObjectKey StaleClass : PreviousClasses)
 		{
-			ClassToSources.FindOrAdd(Object->GetClass()).Add(Object);
+			TSharedPtr<IDetailsView> Details = ClassToDetailsView.FindRef(StaleClass);
+			DetailsBox->RemoveSlot(Details.ToSharedRef());
+			ClassToDetailsView.Remove(StaleClass);
 		}
-	}
 
-	TArray<FObjectKey> PreviousClasses;
-	ClassToDetailsView.GenerateKeyArray(PreviousClasses);
+		PreviousClasses.Empty();
+	}
 
 	for (auto& Pair : ClassToSources)
 	{
diff --git a/Engine/Plugins/VirtualProduction/Takes/Source/TakeTrackRecorders/Private/MovieSceneAnimationTrackRecorder.cpp b/Engine/Plugins/VirtualProduction/Takes/Source/TakeTrackRecorders/Private/MovieSceneAnimationTrackRecorder.cpp
--- a/Engine/Plugins/VirtualProduction/Takes/Source/TakeTrackRecorders/Private/MovieSceneAnimationTrackRecorder.cpp
+++ b/Engine/Plugins/VirtualProduction/Takes/Source/TakeTrackRecorders/Private/MovieSceneAnimationTrackRecorder.cpp
@@ -3,6 +3,7 @@
 #include "TrackRecorders/MovieSceneAnimationTrackRecorder.h"
 #include "TrackRecorders/MovieSceneAnimationTrackRecorderSettings.h"
 #include "TakesUtils.h"
+#include "TakeMetaData.h"
 #include "Tracks/MovieSceneSkeletalAnimationTrack.h"
 #include "Sections/MovieSceneSkeletalAnimationSection.h"
 #include "MovieScene.h"
@@ -36,6 +37,8 @@ UMovieSceneTrackRecorder* FMovieSceneAnimationTrackRecorderFactory::CreateTrackR
 
 void UMovieSceneAnimationTrackRecorder::CreateAnimationAssetAndSequence(const AActor* Actor, const FDirectoryPath& AnimationDirectory)
 {
+	UMovieSceneAnimationTrackRecorderSettings* AnimSettings = CastChecked<UMovieSceneAnimationTrackRecorderSettings>(Settings.Get());
+
 	SkeletalMesh = SkeletalMeshComponent->SkeletalMesh;
 	if (SkeletalMesh.IsValid())
 	{
@@ -44,8 +47,14 @@ void UMovieSceneAnimationTrackRecorder::CreateAnimationAssetAndSequence(const AA
 
 		if (ULevelSequence* MasterLevelSequence = OwningTakeRecorderSource->GetMasterLevelSequence())
 		{
-			const FString& MasterLevelSequenceName = MasterLevelSequence->GetName();
-			AnimationAssetName = AnimationAssetName + TEXT("_") + MasterLevelSequenceName;
+			UTakeMetaData* AssetMetaData = MasterLevelSequence->FindMetaData<UTakeMetaData>();
+
+			AnimationAssetName = AssetMetaData->GenerateAssetPath(AnimSettings->AnimationAssetName);
+
+			TMap<FString, FStringFormatArg> FormatArgs;
+			FormatArgs.Add(TEXT("actor"), Actor->GetActorLabel());
+
+			AnimationAssetName = FString::Format(*AnimationAssetName, FormatArgs);
 		}
 
 		AnimSequence = TakesUtils::MakeNewAsset<UAnimSequence>(AnimationDirectory.Path, AnimationAssetName);
diff --git a/Engine/Plugins/VirtualProduction/Takes/Source/TakeTrackRecorders/Public/TrackRecorders/MovieSceneAnimationTrackRecorderSettings.h b/Engine/Plugins/VirtualProduction/Takes/Source/TakeTrackRecorders/Public/TrackRecorders/MovieSceneAnimationTrackRecorderSettings.h
--- a/Engine/Plugins/VirtualProduction/Takes/Source/TakeTrackRecorders/Public/TrackRecorders/MovieSceneAnimationTrackRecorderSettings.h
+++ b/Engine/Plugins/VirtualProduction/Takes/Source/TakeTrackRecorders/Public/TrackRecorders/MovieSceneAnimationTrackRecorderSettings.h
@@ -16,6 +16,7 @@ public:
 	UMovieSceneAnimationTrackRecorderEditorSettings(const FObjectInitializer& ObjInit)
 		: Super(ObjInit)
 	, AnimationTrackName(NSLOCTEXT("UMovieSceneAnimationTrackRecorderSettings", "DefaultAnimationTrackName", "RecordedAnimation"))
+	, AnimationAssetName(TEXT("{actor}_{slate}_{take}"))
 	, AnimationSubDirectory(TEXT("Animation"))
 	, bRemoveRootAnimation(true)
 	{
@@ -25,6 +26,10 @@ public:
 	UPROPERTY(config, EditAnywhere, BlueprintReadWrite, Category = "Animation Recorder Settings")
 	FText AnimationTrackName;
 
+	/** The name of the animation asset. */
+	UPROPERTY(config, EditAnywhere, BlueprintReadWrite, Category = "Animation Recorder Settings")
+	FString AnimationAssetName;
+
 	/** The name of the subdirectory animations will be placed in. Leave this empty to place into the same directory as the sequence base path. */
 	UPROPERTY(config, EditAnywhere, BlueprintReadWrite, Category = "Animation Recorder Settings")
 	FString AnimationSubDirectory;
diff --git a/Engine/Plugins/VirtualProduction/Takes/Source/TakesCore/Private/TakeRecorderSources.cpp b/Engine/Plugins/VirtualProduction/Takes/Source/TakesCore/Private/TakeRecorderSources.cpp
--- a/Engine/Plugins/VirtualProduction/Takes/Source/TakesCore/Private/TakeRecorderSources.cpp
+++ b/Engine/Plugins/VirtualProduction/Takes/Source/TakesCore/Private/TakeRecorderSources.cpp
@@ -128,6 +128,9 @@ void UTakeRecorderSources::StartRecordingRecursive(TArray<UTakeRecorderSource*>
 					SubsceneTrack = CastChecked<UMovieSceneSubTrack>(InMasterSequence->GetMovieScene()->AddMasterTrack(UMovieSceneSubTrack::StaticClass()));
 				}
 
+				// Track should not be transactional during the recording process
+				SubsceneTrack->ClearFlags(RF_Transactional);
+
 				// We create a new sub track for every Source so that we can name the Subtrack after the Source instead of just the sections within it.
 				SubsceneTrack->SetDisplayName(FText::FromString(Source->GetSubsceneTrackName(InMasterSequence)));
 				SubsceneTrack->SetColorTint(Source->TrackTint);
@@ -145,6 +148,9 @@ void UTakeRecorderSources::StartRecordingRecursive(TArray<UTakeRecorderSource*>
 				FFrameNumber RecordStartTime = FFrameNumber(0);
 				UMovieSceneSubSection* NewSubSection = SubsceneTrack->AddSequence(TargetSequence, RecordStartTime, 0);
 
+				// Section should not be transactional during the recording process
+				NewSubSection->ClearFlags(RF_Transactional);
+
 				NewSubSection->SetRowIndex(SubsceneTrack->GetMaxRowIndex() + 1);
 				SubsceneTrack->FixRowIndices();
 
@@ -294,6 +300,9 @@ void UTakeRecorderSources::PreRecordingRecursive(TArray<UTakeRecorderSource*> In
 					SubsceneTrack = CastChecked<UMovieSceneSubTrack>(InMasterSequence->GetMovieScene()->AddMasterTrack(UMovieSceneSubTrack::StaticClass()));
 				}
 
+				// Track should not be transactional during the recording process
+				SubsceneTrack->ClearFlags(RF_Transactional);
+
 				// We create a new sub track for every Source so that we can name the Subtrack after the Source instead of just the sections within it.
 				SubsceneTrack->SetDisplayName(FText::FromString(Source->GetSubsceneTrackName(InMasterSequence)));
 				SubsceneTrack->SetColorTint(Source->TrackTint);
@@ -311,6 +320,9 @@ void UTakeRecorderSources::PreRecordingRecursive(TArray<UTakeRecorderSource*> In
 				FFrameNumber RecordStartTime = FFrameNumber(0);
 				UMovieSceneSubSection* NewSubSection = SubsceneTrack->AddSequence(TargetSequence, RecordStartTime, 0);
 
+				// Section should not be transactional during the recording process
+				NewSubSection->ClearFlags(RF_Transactional);
+
 				NewSubSection->SetRowIndex(SubsceneTrack->GetMaxRowIndex() + 1);
 				SubsceneTrack->FixRowIndices();
 
@@ -653,6 +665,14 @@ void UTakeRecorderSources::StopRecording(class ULevelSequence* InSequence, FTake
 			// Re-enable transactional after recording
 			SubSequence->GetMovieScene()->SetFlags(RF_Transactional);
 		}
+
+		// Re-enable transactional after recording
+		SubSection->SetFlags(RF_Transactional);
+
+		if (UMovieSceneTrack* SubTrack = Cast<UMovieSceneTrack>(SubSection->GetOuter()))
+		{
+			SubTrack->SetFlags(RF_Transactional);
+		}
 	}
 
 	if (TakeRecorderSourcesSettings.bRemoveRedundantTracks)
diff --git a/Engine/Source/Developer/GameplayDebugger/Private/GameplayDebuggerAddonBase.cpp b/Engine/Source/Developer/GameplayDebugger/Private/GameplayDebuggerAddonBase.cpp
--- a/Engine/Source/Developer/GameplayDebugger/Private/GameplayDebuggerAddonBase.cpp
+++ b/Engine/Source/Developer/GameplayDebugger/Private/GameplayDebuggerAddonBase.cpp
@@ -36,7 +36,10 @@ bool FGameplayDebuggerAddonBase::IsSimulateInEditor()
 {
 #if WITH_EDITOR
 	extern UNREALED_API UEditorEngine* GEditor;
-	return GIsEditor && (GEditor->bIsSimulateInEditorQueued || GEditor->bIsSimulatingInEditor);
+	if (GEditor)
+	{
+		return GEditor->IsSimulateInEditorInProgress();
+	}
 #endif
 	return false;
 }
diff --git a/Engine/Source/Developer/MeshBuilder/Private/SkeletalMeshBuilder.cpp b/Engine/Source/Developer/MeshBuilder/Private/SkeletalMeshBuilder.cpp
--- a/Engine/Source/Developer/MeshBuilder/Private/SkeletalMeshBuilder.cpp
+++ b/Engine/Source/Developer/MeshBuilder/Private/SkeletalMeshBuilder.cpp
@@ -114,8 +114,9 @@ bool FSkeletalMeshBuilder::Build(USkeletalMesh* SkeletalMesh, const int32 LODInd
 	//This scope define where we can use the LODModel, after a reduction the LODModel must be requery since it is a new instance
 	{
 		FSkeletalMeshLODModel& BuildLODModel = SkeletalMesh->GetImportedModel()->LODModels[LODIndex];
-		
-		BuildLODModel.RawSkeletalMeshBulkData.LoadRawMesh(SkeletalMeshImportData);
+
+		//Load the imported data
+		SkeletalMesh->LoadLODImportedData(LODIndex, SkeletalMeshImportData);
 
 		TArray<FVector> LODPoints;
 		TArray<SkeletalMeshImportData::FMeshWedge> LODWedges;
@@ -154,7 +155,7 @@ bool FSkeletalMeshBuilder::Build(USkeletalMesh* SkeletalMesh, const int32 LODInd
 		SlowTask.EnterProgressFrame(1.0f, NSLOCTEXT("SkeltalMeshBuilder", "RebuildMorphTarget", "Rebuilding morph targets..."));
 		if (SkeletalMeshImportData.MorphTargetNames.Num() > 0)
 		{
-			FLODUtilities::BuildMorphTargets(SkeletalMesh, SkeletalMeshImportData, LODIndex, !Options.bComputeNormals, !Options.bComputeTangents, Options.bUseMikkTSpace);
+			FLODUtilities::BuildMorphTargets(SkeletalMesh, SkeletalMeshImportData, LODIndex, !Options.bComputeNormals, !Options.bComputeTangents, Options.bUseMikkTSpace, Options.OverlappingThresholds);
 		}
 
 		//Re-apply the alternate skinning it must be after the inline reduction
diff --git a/Engine/Source/Developer/SkeletalMeshUtilitiesCommon/Private/LODUtilities.cpp b/Engine/Source/Developer/SkeletalMeshUtilitiesCommon/Private/LODUtilities.cpp
--- a/Engine/Source/Developer/SkeletalMeshUtilitiesCommon/Private/LODUtilities.cpp
+++ b/Engine/Source/Developer/SkeletalMeshUtilitiesCommon/Private/LODUtilities.cpp
@@ -1151,11 +1151,15 @@ void FLODUtilities::RestoreSkeletalMeshLODImportedData(USkeletalMesh* SkeletalMe
 	TMap<FString, TArray<FMorphTargetDelta>> ImportedBaseLODMorphTargetData;
 	SkeletalMesh->GetImportedModel()->OriginalReductionSourceMeshData[LodIndex]->LoadReductionData(ImportedBaseLODModel, ImportedBaseLODMorphTargetData, SkeletalMesh);
 	{
-		TArray<int32> EmptyLodInfoMaterialMap;
-		ImportedBaseLODModel.UpdateChunkedSectionInfo(SkeletalMesh->GetName(), EmptyLodInfoMaterialMap);
+		if (!SkeletalMesh->IsLODImportedDataBuildAvailable(LodIndex))
+		{
+			TArray<int32> EmptyLodInfoMaterialMap;
+			ImportedBaseLODModel.UpdateChunkedSectionInfo(SkeletalMesh->GetName(), EmptyLodInfoMaterialMap);
+		}
 		//When we restore a LOD we destroy the LODMaterialMap (user manual section material slot assignation)
 		FSkeletalMeshLODInfo* LODInfo = SkeletalMesh->GetLODInfo(LodIndex);
 		LODInfo->LODMaterialMap.Empty();
+		LODInfo->bHasBeenSimplified = false;
 		//Copy the SkeletalMeshLODModel
 		SkeletalMesh->GetImportedModel()->LODModels[LodIndex] = ImportedBaseLODModel;
 		//Copy the morph target deltas
@@ -1567,21 +1571,21 @@ bool FLODUtilities::UpdateAlternateSkinWeights(USkeletalMesh* SkeletalMeshDest,
 	check(SkeletalMeshDest->GetImportedModel());
 	check(SkeletalMeshDest->GetImportedModel()->LODModels.IsValidIndex(LODIndexDest));
 	FSkeletalMeshLODModel& LODModelDest = SkeletalMeshDest->GetImportedModel()->LODModels[LODIndexDest];
-	return UpdateAlternateSkinWeights(LODModelDest, SkeletalMeshDest->GetName(), SkeletalMeshDest->RefSkeleton, ProfileNameDest, LODIndexDest, OverlappingThresholds, ShouldImportNormals, ShouldImportTangents, bUseMikkTSpace, bComputeWeightedNormals);
-}
-
-bool FLODUtilities::UpdateAlternateSkinWeights(FSkeletalMeshLODModel& LODModelDest, const FString SkeletalMeshName, FReferenceSkeleton& RefSkeleton, const FName& ProfileNameDest, int32 LODIndexDest, FOverlappingThresholds OverlappingThresholds, bool ShouldImportNormals, bool ShouldImportTangents, bool bUseMikkTSpace, bool bComputeWeightedNormals)
-{
-	//Ensure log message only once
-	bool bNoMatchMsgDone = false;
-	if (LODModelDest.RawSkeletalMeshBulkData.IsEmpty())
+	if (SkeletalMeshDest->IsLODImportedDataEmpty(LODIndexDest))
 	{
-		UE_LOG(LogLODUtilities, Error, TEXT("Failed to import Skin Weight Profile as the target skeletal mesh (%s) requires reimporting first."), *SkeletalMeshName);
+		UE_LOG(LogLODUtilities, Error, TEXT("Failed to import Skin Weight Profile as the target skeletal mesh (%s) requires reimporting first."), *SkeletalMeshDest->GetName());
 		//Very old asset will not have this data, we cannot add alternate until the asset is reimported
 		return false;
 	}
 	FSkeletalMeshImportData ImportDataDest;
-	LODModelDest.RawSkeletalMeshBulkData.LoadRawMesh(ImportDataDest);
+	SkeletalMeshDest->LoadLODImportedData(LODIndexDest, ImportDataDest);
+	return UpdateAlternateSkinWeights(LODModelDest, ImportDataDest, SkeletalMeshDest->GetName(), SkeletalMeshDest->RefSkeleton, ProfileNameDest, LODIndexDest, OverlappingThresholds, ShouldImportNormals, ShouldImportTangents, bUseMikkTSpace, bComputeWeightedNormals);
+}
+
+bool FLODUtilities::UpdateAlternateSkinWeights(FSkeletalMeshLODModel& LODModelDest, FSkeletalMeshImportData& ImportDataDest, const FString SkeletalMeshName, FReferenceSkeleton& RefSkeleton, const FName& ProfileNameDest, int32 LODIndexDest, FOverlappingThresholds OverlappingThresholds, bool ShouldImportNormals, bool ShouldImportTangents, bool bUseMikkTSpace, bool bComputeWeightedNormals)
+{
+	//Ensure log message only once
+	bool bNoMatchMsgDone = false;
 	int32 PointNumberDest = ImportDataDest.Points.Num();
 	int32 VertexNumberDest = ImportDataDest.Points.Num();
 
@@ -1861,26 +1865,25 @@ bool FLODUtilities::UpdateAlternateSkinWeights(USkeletalMesh* SkeletalMeshDest,
 	check(SkeletalMeshDest->GetImportedModel());
 	check(SkeletalMeshDest->GetImportedModel()->LODModels.IsValidIndex(LODIndexDest));
 	FSkeletalMeshLODModel& LODModelDest = SkeletalMeshDest->GetImportedModel()->LODModels[LODIndexDest];
-	if (LODModelDest.RawSkeletalMeshBulkData.IsEmpty())
+
+	if (SkeletalMeshDest->IsLODImportedDataEmpty(LODIndexDest))
 	{
 		UE_LOG(LogLODUtilities, Error, TEXT("Failed to import Skin Weight Profile as the target skeletal mesh (%s) requires reimporting first."), SkeletalMeshDest ? *SkeletalMeshDest->GetName() : TEXT("NULL"));
 		//Very old asset will not have this data, we cannot add alternate until the asset is reimported
 		return false;
 	}
 	FSkeletalMeshImportData ImportDataDest;
-	LODModelDest.RawSkeletalMeshBulkData.LoadRawMesh(ImportDataDest);
+	SkeletalMeshDest->LoadLODImportedData(LODIndexDest, ImportDataDest);
 	int32 PointNumberDest = ImportDataDest.Points.Num();
 	int32 VertexNumberDest = ImportDataDest.Points.Num();
 
 	//Grab all the source structure
 	check(SkeletalMeshSrc);
-	check(SkeletalMeshSrc->GetImportedModel());
-	check(SkeletalMeshSrc->GetImportedModel()->LODModels.IsValidIndex(LODIndexSrc));
-	FSkeletalMeshLODModel& LODModelSrc = SkeletalMeshSrc->GetImportedModel()->LODModels[LODIndexSrc];
+
 	//The source model is a fresh import and the data need to be there
-	check(!LODModelSrc.RawSkeletalMeshBulkData.IsEmpty());
+	check(!SkeletalMeshSrc->IsLODImportedDataEmpty(LODIndexSrc));
 	FSkeletalMeshImportData ImportDataSrc;
-	LODModelSrc.RawSkeletalMeshBulkData.LoadRawMesh(ImportDataSrc);
+	SkeletalMeshSrc->LoadLODImportedData(LODIndexSrc, ImportDataSrc);
 	
 	//Remove all unnecessary array data from the structure (this will save a lot of memory)
 	ImportDataSrc.KeepAlternateSkinningBuildDataOnly();
@@ -1894,8 +1897,9 @@ bool FLODUtilities::UpdateAlternateSkinWeights(USkeletalMesh* SkeletalMeshDest,
 	}
 	ImportDataDest.AlternateInfluenceProfileNames.Add(ProfileNameDest.ToString());
 	ImportDataDest.AlternateInfluences.Add(ImportDataSrc);
+
 	//Resave the bulk data with the new or refreshed data
-	LODModelDest.RawSkeletalMeshBulkData.SaveRawMesh(ImportDataDest);
+	SkeletalMeshDest->SaveLODImportedData(LODIndexDest, ImportDataDest);
 
 	//Build the alternate buffer with all the data into the bulk
 	return UpdateAlternateSkinWeights(SkeletalMeshDest, ProfileNameDest, LODIndexDest, OverlappingThresholds, ShouldImportNormals, ShouldImportTangents, bUseMikkTSpace, bComputeWeightedNormals);
@@ -2006,6 +2010,7 @@ void FLODUtilities::RegenerateDependentLODs(USkeletalMesh* SkeletalMesh, int32 L
 		for (const auto& Kvp : Dependencies)
 		{
 			SkeletalMesh->Modify();
+			int32 MaxDependentLODIndex = 0;
 			//Use a TQueue which is thread safe, this Queue will be fill by some delegate call from other threads
 			TQueue<FSkeletalMeshLODModel*> LODModelReplaceByReduction;
 
@@ -2014,6 +2019,7 @@ void FLODUtilities::RegenerateDependentLODs(USkeletalMesh* SkeletalMesh, int32 L
 			TMap<int32, TArray<ClothingAssetUtils::FClothingAssetMeshBinding>> PerLODClothingBindings;
 			for (int32 DependentLODIndex : DependentLODs)
 			{
+				MaxDependentLODIndex = FMath::Max(MaxDependentLODIndex, DependentLODIndex);
 				TArray<ClothingAssetUtils::FClothingAssetMeshBinding>& ClothingBindings = PerLODClothingBindings.FindOrAdd(DependentLODIndex);
 				FLODUtilities::UnbindClothingAndBackup(SkeletalMesh, ClothingBindings, DependentLODIndex);
 
@@ -2024,7 +2030,8 @@ void FLODUtilities::RegenerateDependentLODs(USkeletalMesh* SkeletalMesh, int32 L
 					LODModelReplaceByReduction.Enqueue(ReplacedLODModel);
 				});
 			}
-
+			
+			SkeletalMesh->ReserveLODImportData(MaxDependentLODIndex);
 			//Reduce all dependent LOD in same time
 			ParallelFor(DependentLODs.Num(), [&DependentLODs, &SkeletalMesh](int32 IterationIndex)
 			{
@@ -2104,7 +2111,7 @@ public:
 		TArray< FMorphTargetDelta >& InMorphDeltas, TArray<uint32>& InBaseIndexData, TArray< uint32 >& InBaseWedgePointIndices,
 		TMap<uint32, uint32>& InWedgePointToVertexIndexMap, const FOverlappingCorners& InOverlappingCorners,
 		const TSet<uint32> InModifiedPoints, const TMultiMap< int32, int32 >& InWedgeToFaces, const FMeshDataBundle& InMeshDataBundle, const TArray<FVector>& InTangentZ,
-		bool InShouldImportNormals, bool InShouldImportTangents, bool InbUseMikkTSpace)
+		bool InShouldImportNormals, bool InShouldImportTangents, bool InbUseMikkTSpace, const FOverlappingThresholds InThresholds)
 		: LODModel(InLODModel)
 		, RefSkeleton(InRefSkeleton)
 		, BaseImportData(InBaseImportData)
@@ -2122,6 +2129,7 @@ public:
 		, ShouldImportNormals(InShouldImportNormals)
 		, ShouldImportTangents(InShouldImportTangents)
 		, bUseMikkTSpace(InbUseMikkTSpace)
+		, Thresholds(InThresholds)
 	{
 		MeshUtilities = &FModuleManager::Get().LoadModuleChecked<IMeshUtilities>("MeshUtilities");
 	}
@@ -2250,7 +2258,7 @@ public:
 						}
 
 						// check if position actually changed much
-						if (PositionDelta.SizeSquared() > FMath::Square(THRESH_POINTS_ARE_NEAR) ||
+						if (PositionDelta.SizeSquared() > FMath::Square(Thresholds.MorphThresholdPosition) ||
 							// since we can't get imported morphtarget normal from FBX
 							// we can't compare normal unless it's calculated
 							// this is special flag to ignore normal diff
@@ -2326,9 +2334,10 @@ private:
 	bool ShouldImportNormals;
 	bool ShouldImportTangents;
 	bool bUseMikkTSpace;
+	const FOverlappingThresholds Thresholds;
 };
 
-void FLODUtilities::BuildMorphTargets(USkeletalMesh* BaseSkelMesh, FSkeletalMeshImportData &BaseImportData, int32 LODIndex, bool ShouldImportNormals, bool ShouldImportTangents, bool bUseMikkTSpace)
+void FLODUtilities::BuildMorphTargets(USkeletalMesh* BaseSkelMesh, FSkeletalMeshImportData &BaseImportData, int32 LODIndex, bool ShouldImportNormals, bool ShouldImportTangents, bool bUseMikkTSpace, const FOverlappingThresholds& Thresholds)
 {
 	bool bComputeNormals = !ShouldImportNormals || !BaseImportData.bHasNormals;
 	bool bComputeTangents = !ShouldImportTangents || !BaseImportData.bHasTangents;
@@ -2456,7 +2465,7 @@ void FLODUtilities::BuildMorphTargets(USkeletalMesh* BaseSkelMesh, FSkeletalMesh
 
 			FAsyncTask<FAsyncImportMorphTargetWork>* NewWork = new FAsyncTask<FAsyncImportMorphTargetWork>(&BaseLODModel, BaseSkelMesh->RefSkeleton, BaseImportData,
 				MoveTemp(ShapeImportData.Points), *Deltas, BaseIndexData, BaseWedgePointIndices, WedgePointToVertexIndexMap, OverlappingVertices, MoveTemp(ModifiedPoints), WedgeToFaces, MeshDataBundle, TangentZ,
-				ShouldImportNormals, ShouldImportTangents, bUseMikkTSpace);
+				ShouldImportNormals, ShouldImportTangents, bUseMikkTSpace, Thresholds);
 			PendingWork.Add(NewWork);
 
 			NewWork->StartBackgroundTask(GLargeThreadPool);
@@ -2491,7 +2500,7 @@ void FLODUtilities::BuildMorphTargets(USkeletalMesh* BaseSkelMesh, FSkeletalMesh
 		GWarn->StatusUpdate(Index + 1, MorphTargets.Num(), FText::Format(LOCTEXT("BuildingMorphTargetRenderDataStatus", "Building Morph Target Render Data: {NumCompleted} of {NumTasks}"), Args));
 
 		UMorphTarget* MorphTarget = MorphTargets[Index];
-		MorphTarget->PopulateDeltas(*Results[Index], LODIndex, BaseLODModel.Sections, ShouldImportNormals == false);
+		MorphTarget->PopulateDeltas(*Results[Index], LODIndex, BaseLODModel.Sections, ShouldImportNormals == false, false, Thresholds.MorphThresholdPosition);
 
 		// register does mark package as dirty
 		if (MorphTarget->HasValidData())
diff --git a/Engine/Source/Developer/SkeletalMeshUtilitiesCommon/Public/LODUtilities.h b/Engine/Source/Developer/SkeletalMeshUtilitiesCommon/Public/LODUtilities.h
--- a/Engine/Source/Developer/SkeletalMeshUtilitiesCommon/Public/LODUtilities.h
+++ b/Engine/Source/Developer/SkeletalMeshUtilitiesCommon/Public/LODUtilities.h
@@ -25,6 +25,8 @@ struct FSkeletalMeshUpdateContext
 	FExecuteAction				OnLODChanged;
 };
 
+class FSkeletalMeshImportData;
+
 //////////////////////////////////////////////////////////////////////////
 // FLODUtilities
 
@@ -96,7 +98,7 @@ public:
 	/*
 	 * Build the morph targets for the specified LOD. The function use the Morph target data stored in the FSkeletalMeshImportData ImportData structure
 	 */
-	static void BuildMorphTargets(USkeletalMesh* SkeletalMesh, class FSkeletalMeshImportData &ImportData, int32 LODIndex, bool ShouldImportNormals, bool ShouldImportTangents, bool bUseMikkTSpace);
+	static void BuildMorphTargets(USkeletalMesh* SkeletalMesh, class FSkeletalMeshImportData &ImportData, int32 LODIndex, bool ShouldImportNormals, bool ShouldImportTangents, bool bUseMikkTSpace, const FOverlappingThresholds& Thresholds);
 
 	/**
 	 *	This function apply the skinning weights from asource skeletal mesh to the destination skeletal mesh.
@@ -121,7 +123,7 @@ public:
 	 */
 	static bool UpdateAlternateSkinWeights(USkeletalMesh* SkeletalMeshDest, const FName& ProfileNameDest, int32 LODIndexDest, FOverlappingThresholds OverlappingThresholds, bool ShouldImportNormals, bool ShouldImportTangents, bool bUseMikkTSpace, bool bComputeWeightedNormals);
 
-	static bool UpdateAlternateSkinWeights(FSkeletalMeshLODModel& LODModelDest, const FString SkeletalMeshName, FReferenceSkeleton& RefSkeleton, const FName& ProfileNameDest, int32 LODIndexDest, FOverlappingThresholds OverlappingThresholds, bool ShouldImportNormals, bool ShouldImportTangents, bool bUseMikkTSpace, bool bComputeWeightedNormals);
+	static bool UpdateAlternateSkinWeights(FSkeletalMeshLODModel& LODModelDest, FSkeletalMeshImportData& ImportDataDest, const FString SkeletalMeshName, FReferenceSkeleton& RefSkeleton, const FName& ProfileNameDest, int32 LODIndexDest, FOverlappingThresholds OverlappingThresholds, bool ShouldImportNormals, bool ShouldImportTangents, bool bUseMikkTSpace, bool bComputeWeightedNormals);
 
 	/** Re-generate all (editor-only) skin weight profile, used whenever we rebuild the skeletal mesh data which could change the chunking and bone indices */
 	static void RegenerateAllImportSkinWeightProfileData(FSkeletalMeshLODModel& LODModelDest);
diff --git a/Engine/Source/Editor/ActorPickerMode/Private/EditorModeActorPicker.cpp b/Engine/Source/Editor/ActorPickerMode/Private/EditorModeActorPicker.cpp
--- a/Engine/Source/Editor/ActorPickerMode/Private/EditorModeActorPicker.cpp
+++ b/Engine/Source/Editor/ActorPickerMode/Private/EditorModeActorPicker.cpp
@@ -11,12 +11,13 @@
 
 FEdModeActorPicker::FEdModeActorPicker()
 {
-	PickState = EPickState::NotOverViewport;
-	HoveredActor.Reset();
 }
 
-void FEdModeActorPicker::Initialize()
+void FEdModeActorPicker::Enter()
 {
+	FEdMode::Enter();
+	PickState = EPickState::NotOverViewport;
+	HoveredActor.Reset();
 	CursorDecoratorWindow = SWindow::MakeCursorDecorator();
 	FSlateApplication::Get().AddWindow(CursorDecoratorWindow.ToSharedRef(), true);
 	CursorDecoratorWindow->SetContent(
diff --git a/Engine/Source/Editor/ActorPickerMode/Private/EditorModeActorPicker.h b/Engine/Source/Editor/ActorPickerMode/Private/EditorModeActorPicker.h
--- a/Engine/Source/Editor/ActorPickerMode/Private/EditorModeActorPicker.h
+++ b/Engine/Source/Editor/ActorPickerMode/Private/EditorModeActorPicker.h
@@ -39,11 +39,9 @@ public:
 	virtual bool GetCursor(EMouseCursor::Type& OutCursor) const override;
 	virtual bool UsesToolkits() const override;
 	virtual bool IsCompatibleWith(FEditorModeID OtherModeID) const override;
-	virtual void Initialize() override;
-	/** End FEdMode interface */
-
-	/** End the mode */
+	virtual void Enter() override;
 	virtual void Exit() override;
+	/** End FEdMode interface */
 
 	/** Delegate used to display information about picking near the cursor */
 	FText GetCursorDecoratorText() const;
diff --git a/Engine/Source/Editor/AdvancedPreviewScene/Private/AdvancedPreviewSceneModule.cpp b/Engine/Source/Editor/AdvancedPreviewScene/Private/AdvancedPreviewSceneModule.cpp
--- a/Engine/Source/Editor/AdvancedPreviewScene/Private/AdvancedPreviewSceneModule.cpp
+++ b/Engine/Source/Editor/AdvancedPreviewScene/Private/AdvancedPreviewSceneModule.cpp
@@ -14,12 +14,14 @@ void FAdvancedPreviewSceneModule::ShutdownModule()
 {
 }
 
-TSharedRef<SWidget> FAdvancedPreviewSceneModule::CreateAdvancedPreviewSceneSettingsWidget(const TSharedRef<FAdvancedPreviewScene>& InPreviewScene, UObject* InAdditionalSettings, const TArray<FDetailCustomizationInfo>& InDetailCustomizations, const TArray<FPropertyTypeCustomizationInfo>& InPropertyTypeCustomizations)
+TSharedRef<SWidget> FAdvancedPreviewSceneModule::CreateAdvancedPreviewSceneSettingsWidget(const TSharedRef<FAdvancedPreviewScene>& InPreviewScene, UObject* InAdditionalSettings, const TArray<FDetailCustomizationInfo>& InDetailCustomizations, const TArray<FPropertyTypeCustomizationInfo>& InPropertyTypeCustomizations, const TArray<FDetailDelegates>& InDelegates)
 {
 	return SNew(SAdvancedPreviewDetailsTab, InPreviewScene)
 		.AdditionalSettings(InAdditionalSettings)
 		.DetailCustomizations(InDetailCustomizations)
-		.PropertyTypeCustomizations(InPropertyTypeCustomizations);
+		.PropertyTypeCustomizations(InPropertyTypeCustomizations)
+		.Delegates(InDelegates);
+
 }
 
 IMPLEMENT_MODULE(FAdvancedPreviewSceneModule, AdvancedPreviewScene);
diff --git a/Engine/Source/Editor/AdvancedPreviewScene/Private/SAdvancedPreviewDetailsTab.cpp b/Engine/Source/Editor/AdvancedPreviewScene/Private/SAdvancedPreviewDetailsTab.cpp
--- a/Engine/Source/Editor/AdvancedPreviewScene/Private/SAdvancedPreviewDetailsTab.cpp
+++ b/Engine/Source/Editor/AdvancedPreviewScene/Private/SAdvancedPreviewDetailsTab.cpp
@@ -54,6 +54,7 @@ void SAdvancedPreviewDetailsTab::Construct(const FArguments& InArgs, const TShar
 	ProfileIndex = PerProjectSettings->AssetViewerProfileIndex;
 	DetailCustomizations = InArgs._DetailCustomizations;
 	PropertyTypeCustomizations = InArgs._PropertyTypeCustomizations;
+	Delegates = InArgs._Delegates;
 
 	CreateSettingsView();
 
@@ -138,6 +139,7 @@ void SAdvancedPreviewDetailsTab::ComboBoxSelectionChanged(TSharedPtr<FString> Ne
 		ProfileIndex = NewSelectionIndex;
 		PerProjectSettings->AssetViewerProfileIndex = ProfileIndex;
 		UpdateSettingsView();	
+		check(PreviewScenePtr.IsValid());
 		PreviewScenePtr.Pin()->SetProfileIndex(ProfileIndex);
 	}
 }
@@ -264,6 +266,11 @@ void SAdvancedPreviewDetailsTab::CreateSettingsView()
 		SettingsView->RegisterInstancedCustomPropertyTypeLayout(PropertyTypeCustomizationInfo.StructName, PropertyTypeCustomizationInfo.OnGetPropertyTypeCustomizationInstance);
 	}
 
+	for (FAdvancedPreviewSceneModule::FDetailDelegates& DetailDelegate : Delegates)
+	{
+		DetailDelegate.OnPreviewSceneChangedDelegate.AddSP(this, &SAdvancedPreviewDetailsTab::OnPreviewSceneChanged);
+	}
+
 	class FDetailRootObjectCustomization : public IDetailRootObjectCustomization
 	{
 		virtual TSharedPtr<SWidget> CustomizeObjectHeader(const UObject* InRootObject) { return SNullWidget::NullWidget; }
@@ -291,4 +298,11 @@ void SAdvancedPreviewDetailsTab::OnAssetViewerSettingsPostUndo()
 	PreviewScenePtr.Pin()->UpdateScene(DefaultSettings->Profiles[ProfileIndex]);
 }
 
+void SAdvancedPreviewDetailsTab::OnPreviewSceneChanged(TSharedRef<FAdvancedPreviewScene> PreviewScene)
+{
+  	PreviewScenePtr = PreviewScene;
+	Refresh();
+
+}
+
 #undef LOCTEXT_NAMESPACE
\ No newline at end of file
diff --git a/Engine/Source/Editor/AdvancedPreviewScene/Public/AdvancedPreviewSceneModule.h b/Engine/Source/Editor/AdvancedPreviewScene/Public/AdvancedPreviewSceneModule.h
--- a/Engine/Source/Editor/AdvancedPreviewScene/Public/AdvancedPreviewSceneModule.h
+++ b/Engine/Source/Editor/AdvancedPreviewScene/Public/AdvancedPreviewSceneModule.h
@@ -26,6 +26,14 @@ public:
 		FOnGetPropertyTypeCustomizationInstance OnGetPropertyTypeCustomizationInstance;
 	};
 
+	DECLARE_MULTICAST_DELEGATE_OneParam(FOnPreviewSceneChanged, TSharedRef<FAdvancedPreviewScene>);
+
+	/** Info about a Delegates to subscribe to */
+	struct FDetailDelegates
+	{
+		FOnPreviewSceneChanged& OnPreviewSceneChangedDelegate;
+	};
+
 	// IModuleInterface implementation
 	virtual void StartupModule() override;
 	virtual void ShutdownModule() override;
@@ -37,7 +45,8 @@ public:
 	 * @param	InAdditionalSettings			Additional settings object to display in the view
 	 * @param	InDetailCustomizations			Customizations to use for this details tab
 	 * @param	InPropertyTypeCustomizations	Customizations to use for this details tab
+	 * @param	InDelegates						Delegates to use for this details tab. Array to match other args
 	 * @return a new widget
 	 */
-	virtual TSharedRef<SWidget> CreateAdvancedPreviewSceneSettingsWidget(const TSharedRef<FAdvancedPreviewScene>& InPreviewScene, UObject* InAdditionalSettings = nullptr, const TArray<FDetailCustomizationInfo>& InDetailCustomizations = TArray<FDetailCustomizationInfo>(), const TArray<FPropertyTypeCustomizationInfo>& InPropertyTypeCustomizations = TArray<FPropertyTypeCustomizationInfo>());
+	virtual TSharedRef<SWidget> CreateAdvancedPreviewSceneSettingsWidget(const TSharedRef<FAdvancedPreviewScene>& InPreviewScene, UObject* InAdditionalSettings = nullptr, const TArray<FDetailCustomizationInfo>& InDetailCustomizations = TArray<FDetailCustomizationInfo>(), const TArray<FPropertyTypeCustomizationInfo>& InPropertyTypeCustomizations = TArray<FPropertyTypeCustomizationInfo>(), const TArray<FDetailDelegates>& InDelegates = TArray<FDetailDelegates>());
 };
diff --git a/Engine/Source/Editor/AdvancedPreviewScene/Public/SAdvancedPreviewDetailsTab.h b/Engine/Source/Editor/AdvancedPreviewScene/Public/SAdvancedPreviewDetailsTab.h
--- a/Engine/Source/Editor/AdvancedPreviewScene/Public/SAdvancedPreviewDetailsTab.h
+++ b/Engine/Source/Editor/AdvancedPreviewScene/Public/SAdvancedPreviewDetailsTab.h
@@ -30,6 +30,9 @@ public:
 	/** Customizations to use for this details tab */
 	SLATE_ARGUMENT(TArray<FAdvancedPreviewSceneModule::FPropertyTypeCustomizationInfo>, PropertyTypeCustomizations)
 
+	/** Delegates to use for this details tab */
+	SLATE_ARGUMENT(TArray<FAdvancedPreviewSceneModule::FDetailDelegates>, Delegates)
+
 	SLATE_END_ARGS()
 
 public:
@@ -63,6 +66,8 @@ protected:
 	TArray<TSharedPtr<FString>> ProfileNames;
 	int32 ProfileIndex;
 
+	void OnPreviewSceneChanged(TSharedRef<FAdvancedPreviewScene> PreviewScene);
+
 	FDelegateHandle RefreshDelegate;
 	FDelegateHandle AddRemoveProfileDelegate;
 	FDelegateHandle PostUndoDelegate;
@@ -72,4 +77,6 @@ protected:
 	TArray<FAdvancedPreviewSceneModule::FDetailCustomizationInfo> DetailCustomizations;
 
 	TArray<FAdvancedPreviewSceneModule::FPropertyTypeCustomizationInfo> PropertyTypeCustomizations;
+
+	TArray<FAdvancedPreviewSceneModule::FDetailDelegates> Delegates;
 };
diff --git a/Engine/Source/Editor/Blutility/Classes/EditorUtilityWidgetBlueprint.h b/Engine/Source/Editor/Blutility/Classes/EditorUtilityWidgetBlueprint.h
--- a/Engine/Source/Editor/Blutility/Classes/EditorUtilityWidgetBlueprint.h
+++ b/Engine/Source/Editor/Blutility/Classes/EditorUtilityWidgetBlueprint.h
@@ -32,7 +32,7 @@ public:
 	TSharedRef<SWidget> CreateUtilityWidget();
 
 	/** Recreate the tab's content on recompile */
-	void RegenerateCreatedTab();
+	void RegenerateCreatedTab(UBlueprint* RecompiledBlueprint);
 	
 	void UpdateRespawnListIfNeeded(TSharedRef<SDockTab> TabBeingClosed);
 
diff --git a/Engine/Source/Editor/Blutility/Private/EditorUtilityWidgetBlueprint.cpp b/Engine/Source/Editor/Blutility/Private/EditorUtilityWidgetBlueprint.cpp
--- a/Engine/Source/Editor/Blutility/Private/EditorUtilityWidgetBlueprint.cpp
+++ b/Engine/Source/Editor/Blutility/Private/EditorUtilityWidgetBlueprint.cpp
@@ -56,7 +56,7 @@ TSharedRef<SDockTab> UEditorUtilityWidgetBlueprint::SpawnEditorUITab(const FSpaw
 	SpawnedTab->SetOnTabClosed(SDockTab::FOnTabClosedCallback::CreateUObject(this, &UEditorUtilityWidgetBlueprint::UpdateRespawnListIfNeeded));
 	CreatedTab = SpawnedTab;
 	
-	GEditor->OnBlueprintReinstanced().AddUObject(this, &UEditorUtilityWidgetBlueprint::RegenerateCreatedTab);
+	OnCompiled().AddUObject(this, &UEditorUtilityWidgetBlueprint::RegenerateCreatedTab);
 	
 	FLevelEditorModule& LevelEditor = FModuleManager::LoadModuleChecked<FLevelEditorModule>("LevelEditor");
 	LevelEditor.OnMapChanged().AddUObject(this, &UEditorUtilityWidgetBlueprint::ChangeTabWorld);
@@ -92,7 +92,7 @@ TSharedRef<SWidget> UEditorUtilityWidgetBlueprint::CreateUtilityWidget()
 	return TabWidget;
 }
 
-void UEditorUtilityWidgetBlueprint::RegenerateCreatedTab()
+void UEditorUtilityWidgetBlueprint::RegenerateCreatedTab(UBlueprint* RecompiledBlueprint)
 {
 	if (CreatedTab.IsValid())
 	{
@@ -118,7 +118,7 @@ void UEditorUtilityWidgetBlueprint::ChangeTabWorld(UWorld* World, EMapChangeType
 	else if (MapChangeType != EMapChangeType::SaveMap)
 	{
 		// Recreate the widget if we are loading a map or opening a new map
-		RegenerateCreatedTab();
+		RegenerateCreatedTab(nullptr);
 	}
 }
 
diff --git a/Engine/Source/Editor/ContentBrowser/Private/SFilterList.cpp b/Engine/Source/Editor/ContentBrowser/Private/SFilterList.cpp
--- a/Engine/Source/Editor/ContentBrowser/Private/SFilterList.cpp
+++ b/Engine/Source/Editor/ContentBrowser/Private/SFilterList.cpp
@@ -108,6 +108,7 @@ class SFilter : public SCompoundWidget
 {
 public:
 	DECLARE_DELEGATE_OneParam( FOnRequestRemove, const TSharedRef<SFilter>& /*FilterToRemove*/ );
+	DECLARE_DELEGATE_OneParam( FOnRequestRemoveAllButThis, const TSharedRef<SFilter>& /*FilterToKeep*/ );
 	DECLARE_DELEGATE_OneParam( FOnRequestEnableOnly, const TSharedRef<SFilter>& /*FilterToEnable*/ );
 	DECLARE_DELEGATE( FOnRequestEnableAll );
 	DECLARE_DELEGATE( FOnRequestDisableAll );
@@ -139,6 +140,9 @@ public:
 		/** Invoked when a request to remove all filters originated from within this filter */
 		SLATE_EVENT( FOnRequestRemoveAll, OnRequestRemoveAll )
 
+		/** Invoked when a request to remove all filters originated from within this filter */
+		SLATE_EVENT( FOnRequestRemoveAllButThis, OnRequestRemoveAllButThis )
+
 	SLATE_END_ARGS()
 
 	/** Constructs this widget with InArgs */
@@ -152,6 +156,7 @@ public:
 		OnRequestEnableAll = InArgs._OnRequestEnableAll;
 		OnRequestDisableAll = InArgs._OnRequestDisableAll;
 		OnRequestRemoveAll = InArgs._OnRequestRemoveAll;
+		OnRequestRemoveAllButThis = InArgs._OnRequestRemoveAllButThis;
 		FrontendFilter = InArgs._FrontendFilter;
 
 		// Get the tooltip and color of the type represented by this filter
@@ -303,7 +308,7 @@ private:
 				);
 
 			MenuBuilder.AddMenuEntry(
-				FText::Format( LOCTEXT("EnableOnlyThisFilter", "Enable this only: {0}"), GetFilterName() ),
+				FText::Format( LOCTEXT("EnableOnlyThisFilter", "Enable Only This: {0}"), GetFilterName() ),
 				LOCTEXT("EnableOnlyThisFilterTooltip", "Enable only this filter from the list."),
 				FSlateIcon(),
 				FUIAction( FExecuteAction::CreateSP(this, &SFilter::EnableOnly) )
@@ -334,6 +339,13 @@ private:
 				FSlateIcon(),
 				FUIAction( FExecuteAction::CreateSP(this, &SFilter::RemoveAllFilters) )
 				);
+
+			MenuBuilder.AddMenuEntry(
+				FText::Format( LOCTEXT("RemoveAllButThisFilter", "Remove All But This: {0}"), GetFilterName() ),
+				LOCTEXT("RemoveAllButThisFilterTooltip", "Remove all other filters except this one from the list."),
+				FSlateIcon(),
+				FUIAction( FExecuteAction::CreateSP(this, &SFilter::RemoveAllButThis) )
+				);
 		}
 		MenuBuilder.EndSection();
 
@@ -352,6 +364,13 @@ private:
 		OnRequestRemove.ExecuteIfBound( Self );
 	}
 
+	/** Remove all but this filter from the filter list. */
+	void RemoveAllButThis()
+	{
+		TSharedRef<SFilter> Self = SharedThis(this);
+		OnRequestRemoveAllButThis.ExecuteIfBound(Self);
+	}
+
 	/** Enables only this filter from the filter list */
 	void EnableOnly()
 	{
@@ -443,6 +462,9 @@ private:
 	/** Invoked when a request to remove all filters originated from within this filter */
 	FOnRequestDisableAll OnRequestRemoveAll;
 
+	/** Invoked when a request to remove all filters except this one originated from within this filter */
+	FOnRequestRemoveAllButThis OnRequestRemoveAllButThis;
+
 	/** true when this filter should be applied to the search */
 	bool bEnabled;
 
@@ -684,6 +706,29 @@ void SFilterList::RemoveAllFilters()
 	}
 }
 
+void SFilterList::RemoveAllButThis(const TSharedRef<SFilter>& FilterToKeep)
+{
+	for (const TSharedRef<SFilter>& Filter : Filters)
+	{
+		if (Filter == FilterToKeep)
+		{
+			continue;
+		}
+
+		if (const TSharedPtr<FFrontendFilter>& FrontendFilter = Filter->GetFrontendFilter())
+		{
+			SetFrontendFilterActive(FrontendFilter.ToSharedRef(), false);
+		}
+	}
+
+	FilterBox->ClearChildren();
+	Filters.Empty();
+
+	AddFilter(FilterToKeep);
+
+	OnFilterChanged.ExecuteIfBound();
+}
+
 void SFilterList::DisableFiltersThatHideAssets(const TArray<FAssetData>& AssetDataList)
 {
 	if (HasAnyFilters())
@@ -951,7 +996,8 @@ TSharedRef<SFilter> SFilterList::AddFilter(const TWeakPtr<IAssetTypeActions>& As
 		.OnRequestEnableOnly(this, &SFilterList::EnableOnlyThisFilter)
 		.OnRequestEnableAll(this, &SFilterList::EnableAllFilters)
 		.OnRequestDisableAll(this, &SFilterList::DisableAllFilters)
-		.OnRequestRemoveAll(this, &SFilterList::RemoveAllFilters);
+		.OnRequestRemoveAll(this, &SFilterList::RemoveAllFilters)
+		.OnRequestRemoveAllButThis(this, &SFilterList::RemoveAllButThis);
 
 	AddFilter( NewFilter );
 
@@ -967,7 +1013,8 @@ TSharedRef<SFilter> SFilterList::AddFilter(const TSharedRef<FFrontendFilter>& Fr
 		.OnRequestRemove(this, &SFilterList::RemoveFilterAndUpdate)
 		.OnRequestEnableAll(this, &SFilterList::EnableAllFilters)
 		.OnRequestDisableAll(this, &SFilterList::DisableAllFilters)
-		.OnRequestRemoveAll(this, &SFilterList::RemoveAllFilters);
+		.OnRequestRemoveAll(this, &SFilterList::RemoveAllFilters)
+		.OnRequestRemoveAllButThis(this, &SFilterList::RemoveAllButThis);
 
 	AddFilter( NewFilter );
 
diff --git a/Engine/Source/Editor/ContentBrowser/Private/SFilterList.h b/Engine/Source/Editor/ContentBrowser/Private/SFilterList.h
--- a/Engine/Source/Editor/ContentBrowser/Private/SFilterList.h
+++ b/Engine/Source/Editor/ContentBrowser/Private/SFilterList.h
@@ -70,6 +70,9 @@ public:
 	/** Removes all filters in the list */
 	void RemoveAllFilters();
 
+	/** Removes all filters in the list except the given one. */
+	void RemoveAllButThis(const TSharedRef<SFilter>& FilterToKeep);
+
 	/** Disables any active filters that would hide the supplied assets */
 	void DisableFiltersThatHideAssets(const TArray<FAssetData>& AssetDataList);
 
diff --git a/Engine/Source/Editor/DetailCustomizations/Private/DataTableCustomization.cpp b/Engine/Source/Editor/DetailCustomizations/Private/DataTableCustomization.cpp
--- a/Engine/Source/Editor/DetailCustomizations/Private/DataTableCustomization.cpp
+++ b/Engine/Source/Editor/DetailCustomizations/Private/DataTableCustomization.cpp
@@ -271,8 +271,9 @@ bool FDataTableCustomizationLayout::ShouldFilterAsset(const struct FAssetData& A
 {
 	if (!RowTypeFilter.IsNone())
 	{
-		const UDataTable* DataTable = Cast<UDataTable>(AssetData.GetAsset());
-		if (DataTable->RowStruct && DataTable->RowStruct->GetFName() == RowTypeFilter)
+		static const FName RowStructureTagName("RowStructure");
+		FName RowStructure;
+		if (AssetData.GetTagValue<FName>(RowStructureTagName, RowStructure) && RowStructure == RowTypeFilter)
 		{
 			return false;
 		}
diff --git a/Engine/Source/Editor/DetailCustomizations/Private/MaterialAttributePropertyDetails.cpp b/Engine/Source/Editor/DetailCustomizations/Private/MaterialAttributePropertyDetails.cpp
--- a/Engine/Source/Editor/DetailCustomizations/Private/MaterialAttributePropertyDetails.cpp
+++ b/Engine/Source/Editor/DetailCustomizations/Private/MaterialAttributePropertyDetails.cpp
@@ -20,7 +20,7 @@ TSharedRef<IDetailCustomization> FMaterialAttributePropertyDetails::MakeInstance
 void FMaterialAttributePropertyDetails::CustomizeDetails(IDetailLayoutBuilder& DetailLayout)
 {
 	// Populate combo boxes with material property list
-	FMaterialAttributeDefinitionMap::GetDisplayNameToIDList(AttributeNameToIDList);
+	FMaterialAttributeDefinitionMap::GetAttributeNameToIDList(AttributeNameToIDList);
 
 	AttributeDisplayNameList.Empty(AttributeNameToIDList.Num());
 	for (const TPair<FString, FGuid>& NameGUIDPair : AttributeNameToIDList)
@@ -103,7 +103,7 @@ void FMaterialAttributePropertyDetails::OnBuildChild(TSharedRef<IPropertyHandle>
 						ChildHandle->GetValueAsFormattedString(IDString);
 						FGuid::ParseExact(IDString, EGuidFormats::Digits, IDValue);
 
-						FString AttributeName = FMaterialAttributeDefinitionMap::GetDisplayName(IDValue);
+						FString AttributeName = FMaterialAttributeDefinitionMap::GetAttributeName(IDValue);
 						return FText::FromString(AttributeName);
 					}
 
diff --git a/Engine/Source/Editor/DetailCustomizations/Private/SkeletalMeshReductionSettingsDetails.cpp b/Engine/Source/Editor/DetailCustomizations/Private/SkeletalMeshReductionSettingsDetails.cpp
--- a/Engine/Source/Editor/DetailCustomizations/Private/SkeletalMeshReductionSettingsDetails.cpp
+++ b/Engine/Source/Editor/DetailCustomizations/Private/SkeletalMeshReductionSettingsDetails.cpp
@@ -88,8 +88,7 @@ void FSkeletalMeshReductionSettingsDetails::CustomizeChildren(TSharedRef<IProper
 		// Only able to do this for LOD1 and above, so only show the property if this is the case
 		if (LODIndex > 0)
 		{
-			const FSkeletalMeshModel* ImportedModel = SkeletalMesh != nullptr ? SkeletalMesh->GetImportedModel() : nullptr;
-			bool AllowInline = ImportedModel != nullptr && ImportedModel->LODModels.IsValidIndex(LODIndex) && !ImportedModel->LODModels[LODIndex].RawSkeletalMeshBulkData.IsEmpty();
+			bool AllowInline = !SkeletalMesh->IsLODImportedDataEmpty(LODIndex);
 			// Add and retrieve the default widgets
 			IDetailPropertyRow& Row = StructBuilder.AddProperty(BaseLODPropertyHandle->AsShared());
 
diff --git a/Engine/Source/Editor/FoliageEdit/Private/FoliageEdMode.cpp b/Engine/Source/Editor/FoliageEdit/Private/FoliageEdMode.cpp
--- a/Engine/Source/Editor/FoliageEdit/Private/FoliageEdMode.cpp
+++ b/Engine/Source/Editor/FoliageEdit/Private/FoliageEdMode.cpp
@@ -263,8 +263,6 @@ FEdModeFoliage::FEdModeFoliage()
 	FName OpacityParamName("OpacityAmount");
 	BrushMID->GetScalarParameterValue(OpacityParamName, DefaultBrushOpacity);
 
-	UICommandList = MakeShareable(new FUICommandList);
-	BindCommands();
 }
 
 void FEdModeFoliage::BindCommands()
@@ -307,7 +305,7 @@ void FEdModeFoliage::BindCommands()
 		FCanExecuteAction(),
 		FIsActionChecked::CreateLambda([=]
 	{
-		return UISettings.GetPaintToolSelected();
+		return UISettings.GetPaintToolSelected() && !UISettings.GetIsInSingleInstantiationMode();
 	}));
 
 	UICommandList->MapAction(
@@ -316,7 +314,7 @@ void FEdModeFoliage::BindCommands()
 		FCanExecuteAction(),
 		FIsActionChecked::CreateLambda([=]
 	{
-		return UISettings.GetReapplyToolSelected();
+		return UISettings.GetReapplyToolSelected() && !UISettings.GetIsInSingleInstantiationMode();
 	}));
 
 	UICommandList->MapAction(
@@ -411,6 +409,9 @@ void FEdModeFoliage::Enter()
 	{
 		Toolkit = MakeShareable(new FFoliageEdModeToolkit);
 		Toolkit->Init(Owner->GetToolkitHost());
+
+		UICommandList = Toolkit->GetToolkitCommands();
+		BindCommands();
 	}
 
 	if (UISettings.GetSelectToolSelected() || UISettings.GetLassoSelectToolSelected())
diff --git a/Engine/Source/Editor/FoliageEdit/Private/FoliageEdMode.h b/Engine/Source/Editor/FoliageEdit/Private/FoliageEdMode.h
--- a/Engine/Source/Editor/FoliageEdit/Private/FoliageEdMode.h
+++ b/Engine/Source/Editor/FoliageEdit/Private/FoliageEdMode.h
@@ -529,16 +529,6 @@ public:
 	typedef TMap<FName, TMap<ULandscapeComponent*, TArray<uint8> > > LandscapeLayerCacheData;
 
 	FSimpleMulticastDelegate OnToolChanged;
-private:
-
-	void BindCommands();
-	bool CurrentToolUsesBrush() const;
-
-	/** Called when the user changes the current tool in the UI */
-	void HandleToolChanged();
-
-	/** Deselects all tools */
-	void ClearAllToolSelection();
 
 	/** Sets the tool mode to Paint. */
 	void OnSetPaint();
@@ -555,6 +545,20 @@ private:
 	/** Sets the tool mode to Paint Bucket. */
 	void OnSetPaintFill();
 
+	/** Remove currently selected instances*/
+	void RemoveSelectedInstances(UWorld* InWorld);
+			
+private:
+
+	void BindCommands();
+	bool CurrentToolUsesBrush() const;
+
+	/** Called when the user changes the current tool in the UI */
+	void HandleToolChanged();
+
+	/** Deselects all tools */
+	void ClearAllToolSelection();
+
 	/** Add instances inside the brush to match DesiredInstanceCount */
 	void AddInstancesForBrush(UWorld* InWorld, const UFoliageType* Settings, const FSphere& BrushSphere, int32 DesiredInstanceCount, float Pressure);
 
@@ -605,9 +609,6 @@ private:
 	/**  Updates ed mode widget location to currently selected instance */
 	void UpdateWidgetLocationToInstanceSelection();
 
-	/** Remove currently selected instances*/
-	void RemoveSelectedInstances(UWorld* InWorld);
-			
 	/** Snap instance to the ground   */
 	bool SnapInstanceToGround(AInstancedFoliageActor* InIFA, float AlignMaxAngle, FFoliageInfo& Mesh, int32 InstanceIdx);
 	void SnapSelectedInstancesToGround(UWorld* InWorld);
diff --git a/Engine/Source/Editor/FoliageEdit/Private/FoliageEdModeToolkit.cpp b/Engine/Source/Editor/FoliageEdit/Private/FoliageEdModeToolkit.cpp
--- a/Engine/Source/Editor/FoliageEdit/Private/FoliageEdModeToolkit.cpp
+++ b/Engine/Source/Editor/FoliageEdit/Private/FoliageEdModeToolkit.cpp
@@ -4,9 +4,20 @@
 #include "EditorModeManager.h"
 #include "EditorModes.h"
 #include "SFoliageEdit.h"
+#include "Classes/EditorStyleSettings.h"
+
+#include "FoliageEditActions.h"
+#include "FoliagePaletteCommands.h"
+#include "Framework/MultiBox/MultiBoxBuilder.h"
 
 #define LOCTEXT_NAMESPACE "FoliageEditMode"
 
+namespace 
+{
+	static const FName FoliageName(TEXT("Foliage")); 
+	const TArray<FName> FoliagePaletteNames = { FoliageName };
+}
+
 void FFoliageEdModeToolkit::RegisterTabSpawners(const TSharedRef<class FTabManager>& TabManager)
 {
 
@@ -31,7 +42,7 @@ FName FFoliageEdModeToolkit::GetToolkitFName() const
 
 FText FFoliageEdModeToolkit::GetBaseToolkitName() const
 {
-	return LOCTEXT( "ToolkitName", "Foliage Edit Mode" );
+	return LOCTEXT( "ToolkitName", "Foliage" );
 }
 
 class FEdMode* FFoliageEdModeToolkit::GetEditorMode() const
@@ -54,4 +65,37 @@ void FFoliageEdModeToolkit::NotifyFoliageTypeMeshChanged(class UFoliageType* Fol
 	FoliageEdWidget->NotifyFoliageTypeMeshChanged(FoliageType);
 }
 
+void FFoliageEdModeToolkit::GetToolPaletteNames(TArray<FName>& InPaletteName) const
+{
+	if (!GetDefault<UEditorStyleSettings>()->bEnableLegacyEditorModeUI)
+	{
+		InPaletteName = FoliagePaletteNames;
+	}
+}
+
+FText FFoliageEdModeToolkit::GetToolPaletteDisplayName(FName PaletteName) const
+{
+	if (PaletteName == FoliageName)
+	{
+		return LOCTEXT("Foliage", "Foliage");
+	}
+	return FText();
+}
+
+void FFoliageEdModeToolkit::BuildToolPalette(FName PaletteName, class FToolBarBuilder& ToolBarBuilder)
+{
+	if (PaletteName == FoliageName)
+	{
+		FoliageEdWidget->CustomizeToolBarPalette(ToolBarBuilder);
+	}
+}
+
+void FFoliageEdModeToolkit::OnToolPaletteChanged(FName PaletteName) 
+{
+
+}
+
+
+
+
 #undef LOCTEXT_NAMESPACE
diff --git a/Engine/Source/Editor/FoliageEdit/Private/FoliageEdModeToolkit.h b/Engine/Source/Editor/FoliageEdit/Private/FoliageEdModeToolkit.h
--- a/Engine/Source/Editor/FoliageEdit/Private/FoliageEdModeToolkit.h
+++ b/Engine/Source/Editor/FoliageEdit/Private/FoliageEdModeToolkit.h
@@ -27,6 +27,12 @@ public:
 	virtual class FEdMode* GetEditorMode() const override;
 	virtual TSharedPtr<class SWidget> GetInlineContent() const override;
 
+	/** Mode Toolbar Palettes **/
+	virtual void GetToolPaletteNames(TArray<FName>& InPaletteName) const;
+	virtual FText GetToolPaletteDisplayName(FName PaletteName) const; 
+	virtual void BuildToolPalette(FName PaletteName, class FToolBarBuilder& ToolbarBuilder);
+	virtual void OnToolPaletteChanged(FName PaletteName) override;
+
 	void RefreshFullList();
 	void NotifyFoliageTypeMeshChanged(class UFoliageType* FoliageType);
 
