Port dev-editor fix
Fix unsafe multithreaded access to the Linker in FLODUtilities::RegenerateDependentLODs's use of ParallelFor; call ForceBulkDataResident on the main thread before starting ParallelFor.
#rnx

diff --git a/Engine/Source/Developer/SkeletalMeshUtilitiesCommon/Private/LODUtilities.cpp b/Engine/Source/Developer/SkeletalMeshUtilitiesCommon/Private/LODUtilities.cpp
--- a/Engine/Source/Developer/SkeletalMeshUtilitiesCommon/Private/LODUtilities.cpp
+++ b/Engine/Source/Developer/SkeletalMeshUtilitiesCommon/Private/LODUtilities.cpp
@@ -2025,6 +2025,25 @@ void FLODUtilities::RegenerateDependentLODs(USkeletalMesh* SkeletalMesh, int32 L
 				});
 			}
 
+			// Load the BulkData for all dependent LODs; this has to be done on the main thread since it requires access to the Linker and FLinkerLoad::Serialize is not threadsafe
+			FSkeletalMeshModel* SkeletalMeshResource = SkeletalMesh->GetImportedModel();
+			if (SkeletalMeshResource)
+			{
+				FSkeletalMeshLODModel** LODModels = SkeletalMeshResource->LODModels.GetData();
+				const int32 NumLODModels = SkeletalMeshResource->LODModels.Num();
+				for (int32 DependentLODIndex : DependentLODs)
+				{
+					if (DependentLODIndex < NumLODModels)
+					{
+						FSkeletalMeshLODModel* LODModel = LODModels[DependentLODIndex];
+						if (LODModel)
+						{
+							LODModel->RawSkeletalMeshBulkData.GetBulkData().ForceBulkDataResident();
+						}
+					}
+				}
+			}
+
 			//Reduce all dependent LOD in same time
 			ParallelFor(DependentLODs.Num(), [&DependentLODs, &SkeletalMesh](int32 IterationIndex)
 			{
