Fix crash when unbinding cloth after the user remove the asset in the sklemesh editor.

diff --git a/Engine/Source/Developer/SkeletalMeshUtilitiesCommon/Private/LODUtilities.cpp b/Engine/Source/Developer/SkeletalMeshUtilitiesCommon/Private/LODUtilities.cpp
--- a/Engine/Source/Developer/SkeletalMeshUtilitiesCommon/Private/LODUtilities.cpp
+++ b/Engine/Source/Developer/SkeletalMeshUtilitiesCommon/Private/LODUtilities.cpp
@@ -2557,12 +2557,13 @@ void FLODUtilities::UnbindClothingAndBackup(USkeletalMesh* SkeletalMesh, TArray<
 	{
 		if (Binding.LODIndex == LODIndex)
 		{
-			check(Binding.Asset);
-			Binding.Asset->UnbindFromSkeletalMesh(SkeletalMesh, Binding.LODIndex);
-			
 			//Use the UserSectionsData original section index, this will ensure we remap correctly the cloth if the reduction has change the number of sections
 			int32 OriginalDataSectionIndex = LODModel.Sections[Binding.SectionIndex].OriginalDataSectionIndex;
-			Binding.SectionIndex = OriginalDataSectionIndex;
+			if (Binding.Asset)
+			{
+				Binding.Asset->UnbindFromSkeletalMesh(SkeletalMesh, Binding.LODIndex);
+				Binding.SectionIndex = OriginalDataSectionIndex;
+			}
 			
 			FSkelMeshSourceSectionUserData& SectionUserData = LODModel.UserSectionsData.FindChecked(OriginalDataSectionIndex);
 			SectionUserData.ClothingData.AssetGuid = FGuid();
@@ -2595,9 +2596,8 @@ void FLODUtilities::RestoreClothingFromBackup(USkeletalMesh* SkeletalMesh, TArra
 			{
 				continue;
 			}
-			if (Binding.LODIndex == LODIndex)
+			if (Binding.LODIndex == LODIndex && Binding.Asset)
 			{
-				check(Binding.Asset);
 				if (Binding.Asset->BindToSkeletalMesh(SkeletalMesh, Binding.LODIndex, SectionIndex, Binding.AssetInternalLodIndex))
 				{
 					//If successfull set back the section user data
diff --git a/Engine/Source/Editor/UnrealEd/Private/ApexClothingUtils.cpp b/Engine/Source/Editor/UnrealEd/Private/ApexClothingUtils.cpp
--- a/Engine/Source/Editor/UnrealEd/Private/ApexClothingUtils.cpp
+++ b/Engine/Source/Editor/UnrealEd/Private/ApexClothingUtils.cpp
@@ -158,6 +158,12 @@ void RestoreAllClothingSections(USkeletalMesh* SkelMesh, uint32 LODIndex, uint32
 				if(Section.HasClothingData())
 				{
 					ClothingAssetUtils::ClearSectionClothingData(Section);
+					if (FSkelMeshSourceSectionUserData* UserSectionData = LodModel.UserSectionsData.Find(Section.OriginalDataSectionIndex))
+					{
+						UserSectionData->CorrespondClothAssetIndex = INDEX_NONE;
+						UserSectionData->ClothingData.AssetLodIndex = INDEX_NONE;
+						UserSectionData->ClothingData.AssetGuid = FGuid();
+					}
 				}
 			}
 		}
diff --git a/Engine/Source/Runtime/ClothingSystemRuntimeCommon/Private/ClothingAsset.cpp b/Engine/Source/Runtime/ClothingSystemRuntimeCommon/Private/ClothingAsset.cpp
--- a/Engine/Source/Runtime/ClothingSystemRuntimeCommon/Private/ClothingAsset.cpp
+++ b/Engine/Source/Runtime/ClothingSystemRuntimeCommon/Private/ClothingAsset.cpp
@@ -451,6 +451,12 @@ void UClothingAssetCommon::UnbindFromSkeletalMesh(
 			{
 				InSkelMesh->PreEditChange(nullptr);
 				ClothingAssetUtils::ClearSectionClothingData(Section);
+				if (FSkelMeshSourceSectionUserData* UserSectionData = LodModel.UserSectionsData.Find(Section.OriginalDataSectionIndex))
+				{
+					UserSectionData->CorrespondClothAssetIndex = INDEX_NONE;
+					UserSectionData->ClothingData.AssetLodIndex = INDEX_NONE;
+					UserSectionData->ClothingData.AssetGuid = FGuid();
+				}
 				bChangedMesh = true;
 			}
 		}
