Downgrade the bone influence per vertex limit warning log to a message log
#rnx

diff --git a/Engine/Source/Developer/MeshUtilities/Private/MeshUtilities.cpp b/Engine/Source/Developer/MeshUtilities/Private/MeshUtilities.cpp
--- a/Engine/Source/Developer/MeshUtilities/Private/MeshUtilities.cpp
+++ b/Engine/Source/Developer/MeshUtilities/Private/MeshUtilities.cpp
@@ -3871,7 +3871,7 @@ public:
 						}
 						else
 						{
-							UE_LOG(LogSkeletalMesh, Warning, TEXT("%s"), *(TextMessage.ToString()));
+							UE_LOG(LogSkeletalMesh, Display, TEXT("%s"), *(TextMessage.ToString()));
 						}
 					}
 				}
diff --git a/Engine/Source/Developer/SkeletalMeshUtilitiesCommon/Private/LODUtilities.cpp b/Engine/Source/Developer/SkeletalMeshUtilitiesCommon/Private/LODUtilities.cpp
--- a/Engine/Source/Developer/SkeletalMeshUtilitiesCommon/Private/LODUtilities.cpp
+++ b/Engine/Source/Developer/SkeletalMeshUtilitiesCommon/Private/LODUtilities.cpp
@@ -114,7 +114,7 @@ void FLODUtilities::ProcessImportMeshInfluences(const int32 WedgeCount, TArray<S
 	// warn about too many influences
 	if (MaxVertexInfluence > MAX_TOTAL_INFLUENCES)
 	{
-		UE_LOG(LogLODUtilities, Warning, TEXT("Warning skeletal mesh influence count of %d exceeds max count of %d. Influence truncation will occur. Maximum Ignored Weight %f"), MaxVertexInfluence, MAX_TOTAL_INFLUENCES, MaxIgnoredWeight);
+		UE_LOG(LogLODUtilities, Display, TEXT("Skeletal mesh influence count of %d exceeds max count of %d. Influence truncation will occur. Maximum Ignored Weight %f"), MaxVertexInfluence, MAX_TOTAL_INFLUENCES, MaxIgnoredWeight);
 	}
 
 	for (int32 i = 0; i < Influences.Num(); i++)
diff --git a/Engine/Source/Editor/UnrealEd/Private/Fbx/FbxMaterialImport.cpp b/Engine/Source/Editor/UnrealEd/Private/Fbx/FbxMaterialImport.cpp
--- a/Engine/Source/Editor/UnrealEd/Private/Fbx/FbxMaterialImport.cpp
+++ b/Engine/Source/Editor/UnrealEd/Private/Fbx/FbxMaterialImport.cpp
@@ -96,7 +96,7 @@ UTexture* UnFbx::FFbxImporter::ImportTexture(FbxFileTexture* FbxTexture, bool bS
 	}
 	else
 	{
-		UE_LOG(LogFbxMaterialImport, Warning, TEXT("Unable to find Texture file %s"), *AbsoluteFilename);
+		UE_LOG(LogFbxMaterialImport, Display, TEXT("Unable to find Texture file %s"), *AbsoluteFilename);
 	}
 
 	TArray<uint8> DataBinary;
