[UE-85910] Fix unsafe multithreaded access to the Linker in FLODUtilities::RegenerateDependentLODs's use of ParallelFor; call ForceBulkDataResident on the main thread before starting ParallelFor.
#rnx

diff --git a/Engine/Source/Developer/SkeletalMeshUtilitiesCommon/Private/LODUtilities.cpp b/Engine/Source/Developer/SkeletalMeshUtilitiesCommon/Private/LODUtilities.cpp
--- a/Engine/Source/Developer/SkeletalMeshUtilitiesCommon/Private/LODUtilities.cpp
+++ b/Engine/Source/Developer/SkeletalMeshUtilitiesCommon/Private/LODUtilities.cpp
@@ -2025,7 +2025,23 @@ void FLODUtilities::RegenerateDependentLODs(USkeletalMesh* SkeletalMesh, int32 L
 				});
 			}
 
-			//Reduce all dependent LOD in same time
+			// Load the BulkData for all dependent LODS; this has to be done on the main thread since it requires access to the Linker and FLinkerLoad::Serialize is not threadsafe
+			{
+				FSkeletalMeshModel* SkeletalMeshResource = SkeletalMesh->GetImportedModel();
+				check(SkeletalMeshResource);
+				FSkeletalMeshLODModel** LODModels = SkeletalMeshResource->LODModels.GetData();
+				for (int32 DependentLODIndex : DependentLODs)
+				{
+					check(DependentLODIndex <= SkeletalMeshResource->LODModels.Num());
+					FSkeletalMeshLODModel* LODModel = LODModels[DependentLODIndex];
+					if (LODModel)
+					{
+						LODModel->RawSkeletalMeshBulkData.GetBulkData().ForceBulkDataResident();
+					}
+				}
+			}
+
+			//Reduce all dependent LODs in parallel
 			ParallelFor(DependentLODs.Num(), [&DependentLODs, &SkeletalMesh](int32 IterationIndex)
 			{
 				check(DependentLODs.IsValidIndex(IterationIndex));
