Fix bad assumptions in previous change that could dereference null or use an out of bounds array index.
#rnx

diff --git a/Engine/Source/Developer/SkeletalMeshUtilitiesCommon/Private/LODUtilities.cpp b/Engine/Source/Developer/SkeletalMeshUtilitiesCommon/Private/LODUtilities.cpp
--- a/Engine/Source/Developer/SkeletalMeshUtilitiesCommon/Private/LODUtilities.cpp
+++ b/Engine/Source/Developer/SkeletalMeshUtilitiesCommon/Private/LODUtilities.cpp
@@ -2025,18 +2025,23 @@ void FLODUtilities::RegenerateDependentLODs(USkeletalMesh* SkeletalMesh, int32 L
 				});
 			}
 
-			// Load the BulkData for all dependent LODS; this has to be done on the main thread since it requires access to the Linker and FLinkerLoad::Serialize is not threadsafe
+			// Load the BulkData for all dependent LODs; this has to be done on the main thread since it requires access to the Linker and FLinkerLoad::Serialize is not threadsafe
 			{
 				FSkeletalMeshModel* SkeletalMeshResource = SkeletalMesh->GetImportedModel();
-				check(SkeletalMeshResource);
-				FSkeletalMeshLODModel** LODModels = SkeletalMeshResource->LODModels.GetData();
-				for (int32 DependentLODIndex : DependentLODs)
+				if (SkeletalMeshResource)
 				{
-					check(DependentLODIndex <= SkeletalMeshResource->LODModels.Num());
-					FSkeletalMeshLODModel* LODModel = LODModels[DependentLODIndex];
-					if (LODModel)
+					FSkeletalMeshLODModel** LODModels = SkeletalMeshResource->LODModels.GetData();
+					const int32 NumLODModels = SkeletalMeshResource->LODModels.Num();
+					for (int32 DependentLODIndex : DependentLODs)
 					{
-						LODModel->RawSkeletalMeshBulkData.GetBulkData().ForceBulkDataResident();
+						if (DependentLODIndex < NumLODModels)
+						{
+							FSkeletalMeshLODModel* LODModel = LODModels[DependentLODIndex];
+							if (LODModel)
+							{
+								LODModel->RawSkeletalMeshBulkData.GetBulkData().ForceBulkDataResident();
+							}
+						}
 					}
 				}
 			}
