#UE4 Fix thread safety issue when generating lods

diff --git a/Engine/Source/Developer/SkeletalMeshUtilitiesCommon/Private/LODUtilities.cpp b/Engine/Source/Developer/SkeletalMeshUtilitiesCommon/Private/LODUtilities.cpp
--- a/Engine/Source/Developer/SkeletalMeshUtilitiesCommon/Private/LODUtilities.cpp
+++ b/Engine/Source/Developer/SkeletalMeshUtilitiesCommon/Private/LODUtilities.cpp
@@ -17,6 +17,7 @@
 #include "ClothingAsset.h"
 #include "OverlappingCorners.h"
 #include "Framework/Commands/UIAction.h"
+#include "HAL/ThreadSafeBool.h"
 
 #include "ObjectTools.h"
 #include "ScopedTransaction.h"
@@ -968,7 +969,7 @@ void FLODUtilities::ApplyMorphTargetsToLOD(USkeletalMesh* SkeletalMesh, int32 So
 	CreateLODMorphTarget(SkeletalMesh, ReductionBaseSkeletalMeshBulkData, SourceLOD, DestinationLOD, PerMorphTargetBaseIndexToMorphTargetDelta, BaseMorphIndexToTargetIndexList, TargetVertices, TargetMatchData);
 }
 
-void FLODUtilities::SimplifySkeletalMeshLOD( USkeletalMesh* SkeletalMesh, int32 DesiredLOD, bool bRestoreClothing /*= false*/)
+void FLODUtilities::SimplifySkeletalMeshLOD( USkeletalMesh* SkeletalMesh, int32 DesiredLOD, bool bRestoreClothing /*= false*/, FThreadSafeBool* OutNeedsPackageDirtied/*= nullptr*/)
 {
 	IMeshReductionModule& ReductionModule = FModuleManager::Get().LoadModuleChecked<IMeshReductionModule>("MeshReductionInterface");
 	IMeshReduction* MeshReduction = ReductionModule.GetSkeletalMeshReductionInterface();
@@ -1092,7 +1093,14 @@ void FLODUtilities::SimplifySkeletalMeshLOD( USkeletalMesh* SkeletalMesh, int32
 		};
 
 		ApplyMorphTargetOption();
-		SkeletalMesh->MarkPackageDirty();
+		if (IsInGameThread())
+		{
+			SkeletalMesh->MarkPackageDirty();
+		}
+		else if(OutNeedsPackageDirtied)
+		{
+			(*OutNeedsPackageDirtied) = true;
+		}
 	}
 	else
 	{
@@ -1115,7 +1123,7 @@ void FLODUtilities::SimplifySkeletalMeshLOD( USkeletalMesh* SkeletalMesh, int32
 	}
 }
 
-void FLODUtilities::SimplifySkeletalMeshLOD(FSkeletalMeshUpdateContext& UpdateContext, int32 DesiredLOD, bool bRestoreClothing /*= false*/)
+void FLODUtilities::SimplifySkeletalMeshLOD(FSkeletalMeshUpdateContext& UpdateContext, int32 DesiredLOD, bool bRestoreClothing /*= false*/, FThreadSafeBool* OutNeedsPackageDirtied/*= nullptr*/)
 {
 	USkeletalMesh* SkeletalMesh = UpdateContext.SkeletalMesh;
 	IMeshReductionModule& ReductionModule = FModuleManager::Get().LoadModuleChecked<IMeshReductionModule>("MeshReductionInterface");
@@ -1123,7 +1131,7 @@ void FLODUtilities::SimplifySkeletalMeshLOD(FSkeletalMeshUpdateContext& UpdateCo
 
 	if (MeshReduction && MeshReduction->IsSupported() && SkeletalMesh)
 	{
-		SimplifySkeletalMeshLOD(SkeletalMesh, DesiredLOD, bRestoreClothing);
+		SimplifySkeletalMeshLOD(SkeletalMesh, DesiredLOD, bRestoreClothing, OutNeedsPackageDirtied);
 		
 		if (UpdateContext.OnLODChanged.IsBound())
 		{
@@ -2050,14 +2058,20 @@ void FLODUtilities::RegenerateDependentLODs(USkeletalMesh* SkeletalMesh, int32 L
 
 			SkeletalMesh->ReserveLODImportData(MaxDependentLODIndex);
 			//Reduce all dependent LODs in parallel
-			ParallelFor(DependentLODs.Num(), [&DependentLODs, &SkeletalMesh](int32 IterationIndex)
+			FThreadSafeBool bNeedsPackageDirtied(false);
+			ParallelFor(DependentLODs.Num(), [&DependentLODs, &SkeletalMesh, &bNeedsPackageDirtied](int32 IterationIndex)
 			{
 				check(DependentLODs.IsValidIndex(IterationIndex));
 				int32 DependentLODIndex = DependentLODs[IterationIndex];
 				check(SkeletalMesh->GetLODInfo(DependentLODIndex)); //We cannot add a LOD when reducing with multi thread, so check we already have one
-				FLODUtilities::SimplifySkeletalMeshLOD(SkeletalMesh, DependentLODIndex, false);
+				FLODUtilities::SimplifySkeletalMeshLOD(SkeletalMesh, DependentLODIndex, false, &bNeedsPackageDirtied);
 			});
 
+			if (bNeedsPackageDirtied)
+			{
+				SkeletalMesh->MarkPackageDirty();
+			}
+
 			//Restore the clothings and unbind the delegates
 			for (int32 DependentLODIndex : DependentLODs)
 			{
diff --git a/Engine/Source/Developer/SkeletalMeshUtilitiesCommon/Public/LODUtilities.h b/Engine/Source/Developer/SkeletalMeshUtilitiesCommon/Public/LODUtilities.h
--- a/Engine/Source/Developer/SkeletalMeshUtilitiesCommon/Public/LODUtilities.h
+++ b/Engine/Source/Developer/SkeletalMeshUtilitiesCommon/Public/LODUtilities.h
@@ -68,7 +68,7 @@ public:
 	* @param UpdateContext - The skeletal mesh and actor components to operate on.
 	* @param DesiredLOD - The LOD to simplify
 	*/
-	static void SimplifySkeletalMeshLOD(FSkeletalMeshUpdateContext& UpdateContext, int32 DesiredLOD, bool bRestoreClothing = false);
+	static void SimplifySkeletalMeshLOD(FSkeletalMeshUpdateContext& UpdateContext, int32 DesiredLOD, bool bRestoreClothing = false, class FThreadSafeBool* OutNeedsPackageDirtied = nullptr);
 
 	/**
 	*	Restore the LOD imported model to the last imported data. Call this function if you want to remove the reduce on the base LOD
@@ -149,7 +149,7 @@ private:
 	 * @param SkeletalMesh - The skeletal mesh and actor components to operate on.
 	 * @param DesiredLOD - Desired LOD
 	 */
-	static void SimplifySkeletalMeshLOD(USkeletalMesh* SkeletalMesh, int32 DesiredLOD, bool bRestoreClothing = false);
+	static void SimplifySkeletalMeshLOD(USkeletalMesh* SkeletalMesh, int32 DesiredLOD, bool bRestoreClothing = false, class FThreadSafeBool* OutNeedsPackageDirtied = nullptr);
 
 	/**
 	*  Remap the morph targets of the base LOD onto the desired LOD.
