Merging //UE4/Main[REDACTED_HANDLE] to Dev-RenderPlat-Staging

diff --git a/Engine/Source/Developer/SkeletalMeshUtilitiesCommon/Private/LODUtilities.cpp b/Engine/Source/Developer/SkeletalMeshUtilitiesCommon/Private/LODUtilities.cpp
--- a/Engine/Source/Developer/SkeletalMeshUtilitiesCommon/Private/LODUtilities.cpp
+++ b/Engine/Source/Developer/SkeletalMeshUtilitiesCommon/Private/LODUtilities.cpp
@@ -14,8 +14,8 @@
 #include "Engine/SkeletalMesh.h"
 #include "EditorFramework/AssetImportData.h"
 #include "MeshUtilities.h"
+#include "ClothingAsset.h"
 #include "OverlappingCorners.h"
-#include "Assets/ClothingAsset.h"
 #include "Framework/Commands/UIAction.h"
 
 #include "ObjectTools.h"
@@ -202,14 +202,11 @@ bool FLODUtilities::RegenerateLOD(USkeletalMesh* SkeletalMesh, int32 NewLODCount
 {
 	if (SkeletalMesh)
 	{
+		FScopedSkeletalMeshPostEditChange ScopedPostEditChange(SkeletalMesh);
+
 		// Unbind any existing clothing assets before we regenerate all LODs
 		TArray<ClothingAssetUtils::FClothingAssetMeshBinding> ClothingBindings;
-		ClothingAssetUtils::GetMeshClothingAssetBindings(SkeletalMesh, ClothingBindings);
-
-		for (ClothingAssetUtils::FClothingAssetMeshBinding& Binding : ClothingBindings)
-		{
-			Binding.Asset->UnbindFromSkeletalMesh(SkeletalMesh, Binding.LODIndex);
-		}
+		FLODUtilities::UnbindClothingAndBackup(SkeletalMesh, ClothingBindings);
 
 		int32 LODCount = SkeletalMesh->GetLODNum();
 
@@ -239,7 +236,7 @@ bool FLODUtilities::RegenerateLOD(USkeletalMesh* SkeletalMesh, int32 NewLODCount
 			for (int32 LODIdx = CurrentNumLODs; LODIdx < LODCount; LODIdx++)
 			{
 				// if no previous setting found, it will use default setting. 
-				FLODUtilities::SimplifySkeletalMeshLOD(UpdateContext, LODIdx);
+				FLODUtilities::SimplifySkeletalMeshLOD(UpdateContext, LODIdx, false);
 			}
 		}
 		else
@@ -249,22 +246,13 @@ bool FLODUtilities::RegenerateLOD(USkeletalMesh* SkeletalMesh, int32 NewLODCount
 				FSkeletalMeshLODInfo& CurrentLODInfo = *(SkeletalMesh->GetLODInfo(LODIdx));
 				if ((bRegenerateEvenIfImported && LODIdx > 0) || (bGenerateBaseLOD && LODIdx == 0) || CurrentLODInfo.bHasBeenSimplified )
 				{
-					FLODUtilities::SimplifySkeletalMeshLOD(UpdateContext, LODIdx);
+					FLODUtilities::SimplifySkeletalMeshLOD(UpdateContext, LODIdx, false);
 				}
 			}
 		}
 
 		//Restore all clothing we can
-		for (ClothingAssetUtils::FClothingAssetMeshBinding& Binding : ClothingBindings)
-		{
-			if (SkeletalMesh->GetImportedModel()->LODModels.IsValidIndex(Binding.LODIndex) &&
-				SkeletalMesh->GetImportedModel()->LODModels[Binding.LODIndex].Sections.IsValidIndex(Binding.SectionIndex))
-			{
-				Binding.Asset->BindToSkeletalMesh(SkeletalMesh, Binding.LODIndex, Binding.SectionIndex, Binding.AssetInternalLodIndex);
-			}
-		}
-
-		SkeletalMesh->PostEditChange();
+		FLODUtilities::RestoreClothingFromBackup(SkeletalMesh, ClothingBindings);
 
 		return true;
 	}
@@ -296,21 +284,12 @@ void FLODUtilities::RemoveLOD(FSkeletalMeshUpdateContext& UpdateContext, int32 D
 	// If its a valid LOD, kill it.
 	if( DesiredLOD > 0 && DesiredLOD < SkelMeshModel->LODModels.Num() )
 	{
-		//We'll be modifying the skel mesh data so reregister
-
-		//TODO - do we need to reregister something else instead?
-		FMultiComponentReregisterContext ReregisterContext(UpdateContext.AssociatedComponents);
-
-		// Release rendering resources before deleting LOD
-		SkeletalMesh->ReleaseResources();
+		FScopedSkeletalMeshPostEditChange ScopedPostEditChange(SkeletalMesh);
 
 		// Block until this is done
-		FlushRenderingCommands();
 
 		SkelMeshModel->LODModels.RemoveAt(DesiredLOD);
 		SkeletalMesh->RemoveLODInfo(DesiredLOD);
-		SkeletalMesh->InitResources();
-
 		RefreshLODChange(SkeletalMesh);
 
 		// Set the forced LOD to Auto.
@@ -332,9 +311,6 @@ void FLODUtilities::RemoveLOD(FSkeletalMeshUpdateContext& UpdateContext, int32 D
 			}
 		}
 
-		// This will recache derived render data, and re-init resources
-		SkeletalMesh->PostEditChange();
-
 		//Notify calling system of change
 		UpdateContext.OnLODChanged.ExecuteIfBound();
 
@@ -488,7 +464,7 @@ void ProjectTargetOnBase(const TArray<FSoftSkinVertex>& BaseVertices, const TArr
 	{
 		//Use the remap base index in case some sections disappear during the reduce phase
 		int32 BaseSectionIndex = TargetSectionMatchBaseIndex[SectionIndex];
-		if (BaseSectionIndex == INDEX_NONE || !PerSectionBaseTriangleIndices.IsValidIndex(BaseSectionIndex))
+		if (BaseSectionIndex == INDEX_NONE || !PerSectionBaseTriangleIndices.IsValidIndex(BaseSectionIndex) || PerSectionBaseTriangleIndices[BaseSectionIndex].Num() < 1)
 		{
 			continue;
 		}
@@ -518,6 +494,7 @@ void ProjectTargetOnBase(const TArray<FSoftSkinVertex>& BaseVertices, const TArr
 			TriangleElement.TriangleIndex = Triangles.Num();
 			Triangles.Add(TriangleElement);
 		}
+		check(!BaseMeshUVBound.GetExtent().IsNearlyZero());
 		//Setup the Quad tree
 		float UVsQuadTreeMinSize = 0.001f;
 		TQuadTree<uint32, 100> QuadTree(BaseMeshUVBound, UVsQuadTreeMinSize);
@@ -661,7 +638,7 @@ void CreateLODMorphTarget(USkeletalMesh* SkeletalMesh, FReductionBaseSkeletalMes
 	if (ReductionBaseSkeletalMeshBulkData != nullptr)
 	{
 		FSkeletalMeshLODModel TempBaseLODModel;
-		ReductionBaseSkeletalMeshBulkData->LoadReductionData(TempBaseLODModel, BaseLODMorphTargetData);
+		ReductionBaseSkeletalMeshBulkData->LoadReductionData(TempBaseLODModel, BaseLODMorphTargetData, SkeletalMesh);
 	}
 
 	FSkeletalMeshModel* SkeletalMeshModel = SkeletalMesh->GetImportedModel();
@@ -841,7 +818,8 @@ void FLODUtilities::ApplyMorphTargetsToLOD(USkeletalMesh* SkeletalMesh, int32 So
 	TMap<FString, TArray<FMorphTargetDelta>> TempBaseLODMorphTargetData;
 	if (bReduceBaseLOD)
 	{
-		ReductionBaseSkeletalMeshBulkData->LoadReductionData(TempBaseLODModel, TempBaseLODMorphTargetData);
+		check(ReductionBaseSkeletalMeshBulkData != nullptr);
+		ReductionBaseSkeletalMeshBulkData->LoadReductionData(TempBaseLODModel, TempBaseLODMorphTargetData, SkeletalMesh);
 	}
 
 	const FSkeletalMeshLODModel& BaseLODModel = bReduceBaseLOD ? TempBaseLODModel : SkeletalMeshResource->LODModels[SourceLOD];
@@ -900,7 +878,7 @@ void FLODUtilities::ApplyMorphTargetsToLOD(USkeletalMesh* SkeletalMesh, int32 So
 	for (int32 SectionIndex = 0; SectionIndex < SectionCount; ++SectionIndex)
 	{
 		const FSkelMeshSection& Section = BaseLODModel.Sections[SectionIndex];
-		if (Section.HasClothingData())
+		if (Section.ClothingData.AssetGuid.IsValid())
 		{
 			continue;
 		}
@@ -964,13 +942,14 @@ void FLODUtilities::ApplyMorphTargetsToLOD(USkeletalMesh* SkeletalMesh, int32 So
 	CreateLODMorphTarget(SkeletalMesh, ReductionBaseSkeletalMeshBulkData, SourceLOD, DestinationLOD, PerMorphTargetBaseIndexToMorphTargetDelta, BaseMorphIndexToTargetIndexList, TargetVertices, TargetMatchData);
 }
 
-void FLODUtilities::SimplifySkeletalMeshLOD( USkeletalMesh* SkeletalMesh, int32 DesiredLOD, bool bReregisterComponent /*= true*/, bool bRestoreClothing /*= false*/)
+void FLODUtilities::SimplifySkeletalMeshLOD( USkeletalMesh* SkeletalMesh, int32 DesiredLOD, bool bRestoreClothing /*= false*/)
 {
 	IMeshReductionModule& ReductionModule = FModuleManager::Get().LoadModuleChecked<IMeshReductionModule>("MeshReductionInterface");
 	IMeshReduction* MeshReduction = ReductionModule.GetSkeletalMeshReductionInterface();
 
 	check (MeshReduction && MeshReduction->IsSupported());
 
+
 	if (DesiredLOD == 0
 		&& SkeletalMesh->GetLODInfo(DesiredLOD) != nullptr
 		&& SkeletalMesh->GetLODInfo(DesiredLOD)->bHasBeenSimplified
@@ -1000,18 +979,13 @@ void FLODUtilities::SimplifySkeletalMeshLOD( USkeletalMesh* SkeletalMesh, int32
 		GWarn->BeginSlowTask(StatusUpdate, true);
 	}
 
+	FScopedSkeletalMeshPostEditChange ScopedPostEditChange(SkeletalMesh);
+
 	// Unbind DesiredLOD existing clothing assets before we simplify this LOD
 	TArray<ClothingAssetUtils::FClothingAssetMeshBinding> ClothingBindings;
-	if (bRestoreClothing)
+	if (bRestoreClothing && SkeletalMesh->GetImportedModel() && SkeletalMesh->GetImportedModel()->LODModels.IsValidIndex(DesiredLOD))
 	{
-		ClothingAssetUtils::GetMeshClothingAssetBindings(SkeletalMesh, ClothingBindings);
-		for (ClothingAssetUtils::FClothingAssetMeshBinding& Binding : ClothingBindings)
-		{
-			if (DesiredLOD == Binding.LODIndex)
-			{
-				Binding.Asset->UnbindFromSkeletalMesh(SkeletalMesh, Binding.LODIndex);
-			}
-		}
+		FLODUtilities::UnbindClothingAndBackup(SkeletalMesh, ClothingBindings, DesiredLOD);
 	}
 
 	if (SkeletalMesh->GetLODInfo(DesiredLOD) != nullptr)
@@ -1048,8 +1022,20 @@ void FLODUtilities::SimplifySkeletalMeshLOD( USkeletalMesh* SkeletalMesh, int32
 						MorphDeltasArray.Add(MorphDelta);
 					}
 				}
+				
 				//Copy the original SkeletalMesh LODModel
-				SkeletalMeshResource->OriginalReductionSourceMeshData[DesiredLOD]->SaveReductionData(SrcModel, BaseLODMorphTargetData);
+				// Unbind clothing before saving the original data, we must not restore clothing to do inline reduction
+				{
+					TArray<ClothingAssetUtils::FClothingAssetMeshBinding> TemporaryRemoveClothingBindings;
+					FLODUtilities::UnbindClothingAndBackup(SkeletalMesh, TemporaryRemoveClothingBindings, DesiredLOD);
+
+					SkeletalMeshResource->OriginalReductionSourceMeshData[DesiredLOD]->SaveReductionData(SrcModel, BaseLODMorphTargetData, SkeletalMesh);
+
+					if (TemporaryRemoveClothingBindings.Num() > 0)
+					{
+						FLODUtilities::RestoreClothingFromBackup(SkeletalMesh, TemporaryRemoveClothingBindings, DesiredLOD);
+					}
+				}
 
 				if (DesiredLOD == 0)
 				{
@@ -1060,7 +1046,7 @@ void FLODUtilities::SimplifySkeletalMeshLOD( USkeletalMesh* SkeletalMesh, int32
 	}
 	
 
-	if (MeshReduction->ReduceSkeletalMesh(SkeletalMesh, DesiredLOD, bReregisterComponent))
+	if (MeshReduction->ReduceSkeletalMesh(SkeletalMesh, DesiredLOD))
 	{
 		check(SkeletalMesh->GetLODNum() >= 1);
 
@@ -1078,21 +1064,7 @@ void FLODUtilities::SimplifySkeletalMeshLOD( USkeletalMesh* SkeletalMesh, int32
 			}
 		};
 
-		if (bReregisterComponent)
-		{
-			TComponentReregisterContext<USkinnedMeshComponent> ReregisterContext;
-			SkeletalMesh->ReleaseResources();
-			SkeletalMesh->ReleaseResourcesFence.Wait();
-
-			ApplyMorphTargetOption();
-			
-			SkeletalMesh->PostEditChange();
-			SkeletalMesh->InitResources();
-		}
-		else
-		{
-			ApplyMorphTargetOption();
-		}
+		ApplyMorphTargetOption();
 		SkeletalMesh->MarkPackageDirty();
 	}
 	else
@@ -1105,25 +1077,15 @@ void FLODUtilities::SimplifySkeletalMeshLOD( USkeletalMesh* SkeletalMesh, int32
 	}
 
 	//Put back the clothing for the DesiredLOD
-	if (bRestoreClothing)
+	if (bRestoreClothing && ClothingBindings.Num() > 0 && SkeletalMesh->GetImportedModel()->LODModels.IsValidIndex(DesiredLOD))
 	{
-		for (ClothingAssetUtils::FClothingAssetMeshBinding& Binding : ClothingBindings)
-		{
-			if (SkeletalMesh->GetImportedModel()->LODModels.IsValidIndex(Binding.LODIndex) &&
-				SkeletalMesh->GetImportedModel()->LODModels[Binding.LODIndex].Sections.IsValidIndex(Binding.SectionIndex))
-			{
-				if (DesiredLOD == Binding.LODIndex)
-				{
-					Binding.Asset->BindToSkeletalMesh(SkeletalMesh, Binding.LODIndex, Binding.SectionIndex, Binding.AssetInternalLodIndex);
-				}
-			}
-		}
+		FLODUtilities::RestoreClothingFromBackup(SkeletalMesh, ClothingBindings, DesiredLOD);
 	}
 
 	GWarn->EndSlowTask();
 }
 
-void FLODUtilities::SimplifySkeletalMeshLOD(FSkeletalMeshUpdateContext& UpdateContext, int32 DesiredLOD, bool bReregisterComponent /*= true*/, bool bRestoreClothing /*= false*/)
+void FLODUtilities::SimplifySkeletalMeshLOD(FSkeletalMeshUpdateContext& UpdateContext, int32 DesiredLOD, bool bRestoreClothing /*= false*/)
 {
 	USkeletalMesh* SkeletalMesh = UpdateContext.SkeletalMesh;
 	IMeshReductionModule& ReductionModule = FModuleManager::Get().LoadModuleChecked<IMeshReductionModule>("MeshReductionInterface");
@@ -1131,7 +1093,7 @@ void FLODUtilities::SimplifySkeletalMeshLOD(FSkeletalMeshUpdateContext& UpdateCo
 
 	if (MeshReduction && MeshReduction->IsSupported() && SkeletalMesh)
 	{
-		SimplifySkeletalMeshLOD(SkeletalMesh, DesiredLOD, bReregisterComponent, bRestoreClothing);
+		SimplifySkeletalMeshLOD(SkeletalMesh, DesiredLOD, bRestoreClothing);
 		
 		if (UpdateContext.OnLODChanged.IsBound())
 		{
@@ -1141,40 +1103,27 @@ void FLODUtilities::SimplifySkeletalMeshLOD(FSkeletalMeshUpdateContext& UpdateCo
 	}
 }
 
-void FLODUtilities::RestoreSkeletalMeshLODImportedData(USkeletalMesh* SkeletalMesh, int32 LodIndex, bool bReregisterComponent /*= true*/)
+void FLODUtilities::RestoreSkeletalMeshLODImportedData(USkeletalMesh* SkeletalMesh, int32 LodIndex)
 {
-
 	if (!SkeletalMesh->GetImportedModel()->OriginalReductionSourceMeshData.IsValidIndex(LodIndex) || SkeletalMesh->GetImportedModel()->OriginalReductionSourceMeshData[LodIndex]->IsEmpty())
 	{
 		//There is nothing to restore
 		return;
 	}
 
+	FScopedSkeletalMeshPostEditChange ScopedPostEditChange(SkeletalMesh);
+
 	// Unbind LodIndex existing clothing assets before restoring the LOD
 	TArray<ClothingAssetUtils::FClothingAssetMeshBinding> ClothingBindings;
-	ClothingAssetUtils::GetMeshClothingAssetBindings(SkeletalMesh, ClothingBindings);
-	for (ClothingAssetUtils::FClothingAssetMeshBinding& Binding : ClothingBindings)
-	{
-		//Unbind only the LOD we restore
-		if (Binding.LODIndex == LodIndex)
-		{
-			Binding.Asset->UnbindFromSkeletalMesh(SkeletalMesh, Binding.LODIndex);
-		}
-	}
+	FLODUtilities::UnbindClothingAndBackup(SkeletalMesh, ClothingBindings);
 
 	FSkeletalMeshLODModel ImportedBaseLODModel;
 	TMap<FString, TArray<FMorphTargetDelta>> ImportedBaseLODMorphTargetData;
-	SkeletalMesh->GetImportedModel()->OriginalReductionSourceMeshData[LodIndex]->LoadReductionData(ImportedBaseLODModel, ImportedBaseLODMorphTargetData);
+	SkeletalMesh->GetImportedModel()->OriginalReductionSourceMeshData[LodIndex]->LoadReductionData(ImportedBaseLODModel, ImportedBaseLODMorphTargetData, SkeletalMesh);
 	{
 		FSkeletalMeshUpdateContext UpdateContext;
 		UpdateContext.SkeletalMesh = SkeletalMesh;
 
-		TComponentReregisterContext<USkinnedMeshComponent> ReregisterContext;
-		if (bReregisterComponent)
-		{
-			SkeletalMesh->ReleaseResources();
-			SkeletalMesh->ReleaseResourcesFence.Wait();
-		}
 		//Copy the SkeletalMeshLODModel
 		SkeletalMesh->GetImportedModel()->LODModels[LodIndex] = ImportedBaseLODModel;
 		//Copy the morph target deltas
@@ -1196,20 +1145,7 @@ void FLODUtilities::RestoreSkeletalMeshLODImportedData(USkeletalMesh* SkeletalMe
 		SkeletalMesh->GetImportedModel()->OriginalReductionSourceMeshData[LodIndex]->EmptyBulkData();
 
 		//Put back the clothing for the restore LOD
-		for (ClothingAssetUtils::FClothingAssetMeshBinding& Binding : ClothingBindings)
-		{
-			if (LodIndex == Binding.LODIndex && SkeletalMesh->GetImportedModel()->LODModels.IsValidIndex(Binding.LODIndex) &&
-				SkeletalMesh->GetImportedModel()->LODModels[Binding.LODIndex].Sections.IsValidIndex(Binding.SectionIndex))
-			{
-				Binding.Asset->BindToSkeletalMesh(SkeletalMesh, Binding.LODIndex, Binding.SectionIndex, Binding.AssetInternalLodIndex);
-			}
-		}
-
-		if (bReregisterComponent)
-		{
-			SkeletalMesh->PostEditChange();
-			SkeletalMesh->InitResources();
-		}
+		FLODUtilities::RestoreClothingFromBackup(SkeletalMesh, ClothingBindings);
 
 		if (UpdateContext.OnLODChanged.IsBound())
 		{
@@ -1600,17 +1536,21 @@ void MatchVertexIndexUsingPosition(
 
 bool FLODUtilities::UpdateAlternateSkinWeights(USkeletalMesh* SkeletalMeshDest, const FName& ProfileNameDest, int32 LODIndexDest, FOverlappingThresholds OverlappingThresholds, bool ShouldImportNormals, bool ShouldImportTangents, bool bUseMikkTSpace, bool bComputeWeightedNormals)
 {
-	//Ensure log message only once
-	bool bNoMatchMsgDone = false;
-
 	//Grab all the destination structure
 	check(SkeletalMeshDest);
 	check(SkeletalMeshDest->GetImportedModel());
 	check(SkeletalMeshDest->GetImportedModel()->LODModels.IsValidIndex(LODIndexDest));
 	FSkeletalMeshLODModel& LODModelDest = SkeletalMeshDest->GetImportedModel()->LODModels[LODIndexDest];
+	return UpdateAlternateSkinWeights(LODModelDest, SkeletalMeshDest->GetName(), SkeletalMeshDest->RefSkeleton, ProfileNameDest, LODIndexDest, OverlappingThresholds, ShouldImportNormals, ShouldImportTangents, bUseMikkTSpace, bComputeWeightedNormals);
+}
+
+bool FLODUtilities::UpdateAlternateSkinWeights(FSkeletalMeshLODModel& LODModelDest, const FString SkeletalMeshName, FReferenceSkeleton& RefSkeleton, const FName& ProfileNameDest, int32 LODIndexDest, FOverlappingThresholds OverlappingThresholds, bool ShouldImportNormals, bool ShouldImportTangents, bool bUseMikkTSpace, bool bComputeWeightedNormals)
+{
+	//Ensure log message only once
+	bool bNoMatchMsgDone = false;
 	if (LODModelDest.RawSkeletalMeshBulkData.IsEmpty())
 	{
-		UE_LOG(LogLODUtilities, Error, TEXT("Failed to import Skin Weight Profile as the target skeletal mesh (%s) requires reimporting first."), SkeletalMeshDest ? *SkeletalMeshDest->GetName() : TEXT("NULL"));
+		UE_LOG(LogLODUtilities, Error, TEXT("Failed to import Skin Weight Profile as the target skeletal mesh (%s) requires reimporting first."), *SkeletalMeshName);
 		//Very old asset will not have this data, we cannot add alternate until the asset is reimported
 		return false;
 	}
@@ -1622,7 +1562,7 @@ bool FLODUtilities::UpdateAlternateSkinWeights(USkeletalMesh* SkeletalMeshDest,
 	int32 ProfileIndex = 0;
 	if (!ImportDataDest.AlternateInfluenceProfileNames.Find(ProfileNameDest.ToString(), ProfileIndex))
 	{
-		UE_LOG(LogLODUtilities, Error, TEXT("Failed to import Skin Weight Profile the alternate skinning imported source data is not available."), SkeletalMeshDest ? *SkeletalMeshDest->GetName() : TEXT("NULL"));
+		UE_LOG(LogLODUtilities, Error, TEXT("Failed to import Skin Weight Profile the alternate skinning imported source data is not available."), *SkeletalMeshName);
 		return false;
 	}
 	check(ImportDataDest.AlternateInfluences.IsValidIndex(ProfileIndex));
@@ -1634,7 +1574,7 @@ bool FLODUtilities::UpdateAlternateSkinWeights(USkeletalMesh* SkeletalMeshDest,
 
 	if (ImportDataDest.NumTexCoords <= 0 || ImportDataSrc.NumTexCoords <= 0)
 	{
-		UE_LOG(LogLODUtilities, Error, TEXT("Failed to import Skin Weight Profile as the target skeletal mesh (%s) or imported file does not contain UV coordinates."), SkeletalMeshDest ? *SkeletalMeshDest->GetName() : TEXT("NULL"));
+		UE_LOG(LogLODUtilities, Error, TEXT("Failed to import Skin Weight Profile as the target skeletal mesh (%s) or imported file does not contain UV coordinates."), *SkeletalMeshName);
 		return false;
 	}
 
@@ -1650,7 +1590,7 @@ bool FLODUtilities::UpdateAlternateSkinWeights(USkeletalMesh* SkeletalMeshDest,
 
 	//Create a map to remap source bone index to destination bone index
 	TMap<int32, int32> RemapBoneIndexSrcToDest;
-	FillRemapBoneIndexSrcToDest(ImportDataSrc, ImportDataDest, SkeletalMeshDest->GetName(), LODIndexDest, RemapBoneIndexSrcToDest);
+	FillRemapBoneIndexSrcToDest(ImportDataSrc, ImportDataDest, SkeletalMeshName, LODIndexDest, RemapBoneIndexSrcToDest);
 
 	//Map to get the vertex index source to a destination vertex match
 	TSortedMap<uint32, VertexMatchNameSpace::FVertexMatchResult> VertexIndexSrcToVertexIndexDestMatches;
@@ -1879,7 +1819,10 @@ bool FLODUtilities::UpdateAlternateSkinWeights(USkeletalMesh* SkeletalMeshDest,
 	TArray<FText> WarningMessages;
 	TArray<FName> WarningNames;
 	//Build the destination mesh with the Alternate influences, so the chunking is done properly.
-	bool bBuildSuccess = MeshUtilities.BuildSkeletalMesh(LODModelDest, SkeletalMeshDest->RefSkeleton, LODInfluencesDest, LODWedgesDest, LODFacesDest, LODPointsDest, LODPointToRawMapDest, BuildOptions, &WarningMessages, &WarningNames);
+	bool bBuildSuccess = MeshUtilities.BuildSkeletalMesh(LODModelDest, RefSkeleton, LODInfluencesDest, LODWedgesDest, LODFacesDest, LODPointsDest, LODPointToRawMapDest, BuildOptions, &WarningMessages, &WarningNames);
+	//Re-Apply the user section changes, the UserSectionsData is map to original section and should match the builded LODModel
+	LODModelDest.SyncronizeUserSectionsDataArray();
+
 	RegenerateAllImportSkinWeightProfileData(LODModelDest);
 	
 	return bBuildSuccess;
@@ -2004,7 +1947,8 @@ void FLODUtilities::RegenerateDependentLODs(USkeletalMesh* SkeletalMesh, int32 L
 {
 	FSkeletalMeshUpdateContext UpdateContext;
 	UpdateContext.SkeletalMesh = SkeletalMesh;
-	//Check the dependencies and regenerate the LODs acoording to it
+
+	//Check the dependencies and regenerate the LODs according to it
 	TArray<bool> LODDependencies;
 	int32 LODNumber = SkeletalMesh->GetLODNum();
 	LODDependencies.AddZeroed(LODNumber);
@@ -2022,19 +1966,20 @@ void FLODUtilities::RegenerateDependentLODs(USkeletalMesh* SkeletalMesh, int32 L
 	}
 	if (bRegenLODs)
 	{
-		TComponentReregisterContext<USkinnedMeshComponent> ReregisterContext;
-		SkeletalMesh->Modify();
-		SkeletalMesh->ReleaseResources();
-		SkeletalMesh->ReleaseResourcesFence.Wait();
-		for (int32 DependentLODIndex = LODIndex + 1; DependentLODIndex < LODNumber; ++DependentLODIndex)
+		FScopedSkeletalMeshPostEditChange ScopedPostEditChange(SkeletalMesh);
+		auto SimplifyDependentLODS = [&LODIndex, &LODNumber, &LODDependencies, &UpdateContext]()
 		{
-			if (LODDependencies[DependentLODIndex])
+			for (int32 DependentLODIndex = LODIndex + 1; DependentLODIndex < LODNumber; ++DependentLODIndex)
 			{
-				FLODUtilities::SimplifySkeletalMeshLOD(UpdateContext, DependentLODIndex, false);
+				if (LODDependencies[DependentLODIndex])
+				{
+					FLODUtilities::SimplifySkeletalMeshLOD(UpdateContext, DependentLODIndex, true);
+				}
 			}
-		}
-		SkeletalMesh->PostEditChange();
-		SkeletalMesh->InitResources();
+		};
+
+		SkeletalMesh->Modify();
+		SimplifyDependentLODS();
 	}
 }
 
@@ -2486,4 +2431,84 @@ void FLODUtilities::BuildMorphTargets(USkeletalMesh* BaseSkelMesh, FSkeletalMesh
 	}
 }
 
+void FLODUtilities::UnbindClothingAndBackup(USkeletalMesh* SkeletalMesh, TArray<ClothingAssetUtils::FClothingAssetMeshBinding>& ClothingBindings)
+{
+	for (int32 LODIndex = 0; LODIndex < SkeletalMesh->GetImportedModel()->LODModels.Num(); ++LODIndex)
+	{
+		TArray<ClothingAssetUtils::FClothingAssetMeshBinding> LODBindings;
+		UnbindClothingAndBackup(SkeletalMesh, LODBindings, LODIndex);
+		ClothingBindings.Append(LODBindings);
+	}
+}
+
+void FLODUtilities::UnbindClothingAndBackup(USkeletalMesh* SkeletalMesh, TArray<ClothingAssetUtils::FClothingAssetMeshBinding>& ClothingBindings, const int32 LODIndex)
+{
+	if (!SkeletalMesh->GetImportedModel()->LODModels.IsValidIndex(LODIndex))
+	{
+		return;
+	}
+	FSkeletalMeshLODModel& LODModel = SkeletalMesh->GetImportedModel()->LODModels[LODIndex];
+	//Store the clothBinding
+	ClothingAssetUtils::GetMeshClothingAssetBindings(SkeletalMesh, ClothingBindings, LODIndex);
+	//Unbind the Cloth for this LOD before we reduce it, we will put back the cloth after the reduction, if it still match the sections
+	for (ClothingAssetUtils::FClothingAssetMeshBinding& Binding : ClothingBindings)
+	{
+		if (Binding.LODIndex == LODIndex)
+		{
+			check(Binding.Asset);
+			Binding.Asset->UnbindFromSkeletalMesh(SkeletalMesh, Binding.LODIndex);
+			
+			//Use the UserSectionsData original section index, this will ensure we remap correctly the cloth if the reduction has change the number of sections
+			int32 OriginalDataSectionIndex = LODModel.Sections[Binding.SectionIndex].OriginalDataSectionIndex;
+			Binding.SectionIndex = OriginalDataSectionIndex;
+			
+			FSkelMeshSourceSectionUserData& SectionUserData = LODModel.UserSectionsData.FindChecked(OriginalDataSectionIndex);
+			SectionUserData.ClothingData.AssetGuid = FGuid();
+			SectionUserData.ClothingData.AssetLodIndex = INDEX_NONE;
+			SectionUserData.CorrespondClothAssetIndex = INDEX_NONE;
+		}
+	}
+}
+
+void FLODUtilities::RestoreClothingFromBackup(USkeletalMesh* SkeletalMesh, TArray<ClothingAssetUtils::FClothingAssetMeshBinding>& ClothingBindings)
+{
+	for (int32 LODIndex = 0; LODIndex < SkeletalMesh->GetImportedModel()->LODModels.Num(); ++LODIndex)
+	{
+		RestoreClothingFromBackup(SkeletalMesh, ClothingBindings, LODIndex);
+	}
+}
+
+void FLODUtilities::RestoreClothingFromBackup(USkeletalMesh* SkeletalMesh, TArray<ClothingAssetUtils::FClothingAssetMeshBinding>& ClothingBindings, const int32 LODIndex)
+{
+	if (!SkeletalMesh->GetImportedModel()->LODModels.IsValidIndex(LODIndex))
+	{
+		return;
+	}
+	FSkeletalMeshLODModel& LODModel = SkeletalMesh->GetImportedModel()->LODModels[LODIndex];
+	for (ClothingAssetUtils::FClothingAssetMeshBinding& Binding : ClothingBindings)
+	{
+		for (int32 SectionIndex = 0; SectionIndex < LODModel.Sections.Num(); ++SectionIndex)
+		{
+			if (LODModel.Sections[SectionIndex].OriginalDataSectionIndex != Binding.SectionIndex)
+			{
+				continue;
+			}
+			if (Binding.LODIndex == LODIndex)
+			{
+				check(Binding.Asset);
+				if (Binding.Asset->BindToSkeletalMesh(SkeletalMesh, Binding.LODIndex, SectionIndex, Binding.AssetInternalLodIndex))
+				{
+					//If successfull set back the section user data
+					FSkelMeshSourceSectionUserData& SectionUserData = LODModel.UserSectionsData.FindChecked(Binding.SectionIndex);
+					SectionUserData.CorrespondClothAssetIndex = LODModel.Sections[SectionIndex].CorrespondClothAssetIndex;
+					SectionUserData.ClothingData = LODModel.Sections[SectionIndex].ClothingData;
+				}
+			}
+			break;
+		}
+	}
+}
+
+
+
 #undef LOCTEXT_NAMESPACE // "LODUtilities"
\ No newline at end of file
diff --git a/Engine/Source/Editor/ClothPainter/Private/SClothPaintTab.cpp b/Engine/Source/Editor/ClothPainter/Private/SClothPaintTab.cpp
--- a/Engine/Source/Editor/ClothPainter/Private/SClothPaintTab.cpp
+++ b/Engine/Source/Editor/ClothPainter/Private/SClothPaintTab.cpp
@@ -27,7 +27,7 @@
 #include "SClothPaintWidget.h"
 #include "IPersonaToolkit.h"
 #include "SClothAssetSelector.h"
-#include "Assets/ClothingAsset.h"
+#include "ClothingAsset.h"
 #include "Widgets/Layout/SScrollBox.h"
 #include "ComponentReregisterContext.h"
 
@@ -51,7 +51,7 @@ SClothPaintTab::~SClothPaintTab()
 
 void SClothPaintTab::Construct(const FArguments& InArgs)
 {
-	// Detail view for UClothingAsset
+	// Detail view for UClothingAssetCommon
 	FPropertyEditorModule& EditModule = FModuleManager::Get().GetModuleChecked<FPropertyEditorModule>("PropertyEditor");
 
 	FDetailsViewArgs DetailsViewArgs(
@@ -75,7 +75,7 @@ void SClothPaintTab::Construct(const FArguments& InArgs)
 
 	// Add the CDO by default
 	TArray<UObject*> Objects;
-	Objects.Add(UClothingAsset::StaticClass()->GetDefaultObject());
+	Objects.Add(UClothingAssetCommon::StaticClass()->GetDefaultObject());
 	DetailsView->SetObjects(Objects, true);
 
 	HostingApp = InArgs._InHostingApp;
@@ -147,7 +147,7 @@ void SClothPaintTab::UpdatePaintTools()
 		ISkeletalMeshEditor* SkeletalMeshEditor = static_cast<ISkeletalMeshEditor*>(HostingApp.Pin().Get());
 		SkeletalMeshEditor->GetAssetEditorModeManager()->ActivateMode(PaintModeID, true);
 
-		FClothingPaintEditMode* PaintMode = (FClothingPaintEditMode*)SkeletalMeshEditor->GetAssetEditorModeManager()->FindMode(PaintModeID);
+		FClothingPaintEditMode* PaintMode = (FClothingPaintEditMode*)SkeletalMeshEditor->GetAssetEditorModeManager()->GetActiveMode(PaintModeID);
 		if (PaintMode)
 		{
 			FClothPainter* ClothPainter = static_cast<FClothPainter*>(PaintMode->GetMeshPainter());
@@ -165,7 +165,7 @@ void SClothPaintTab::UpdatePaintTools()
 
 			if(SelectorWidget.IsValid())
 			{
-				TWeakObjectPtr<UClothingAsset> WeakAsset = SelectorWidget->GetSelectedAsset();
+				TWeakObjectPtr<UClothingAssetCommon> WeakAsset = SelectorWidget->GetSelectedAsset();
 
 				if(WeakAsset.Get())
 				{
@@ -183,13 +183,13 @@ void SClothPaintTab::UpdatePaintTools()
 	}
 }
 
-void SClothPaintTab::OnAssetSelectionChanged(TWeakObjectPtr<UClothingAsset> InAssetPtr, int32 InLodIndex, int32 InMaskIndex)
+void SClothPaintTab::OnAssetSelectionChanged(TWeakObjectPtr<UClothingAssetCommon> InAssetPtr, int32 InLodIndex, int32 InMaskIndex)
 {
 	if(bPaintModeEnabled)
 	{
 		ISkeletalMeshEditor* SkeletalMeshEditor = static_cast<ISkeletalMeshEditor*>(HostingApp.Pin().Get());
 
-		FClothingPaintEditMode* PaintMode = (FClothingPaintEditMode*)SkeletalMeshEditor->GetAssetEditorModeManager()->FindMode(PaintModeID);
+		FClothingPaintEditMode* PaintMode = (FClothingPaintEditMode*)SkeletalMeshEditor->GetAssetEditorModeManager()->GetActiveMode(PaintModeID);
 		if(PaintMode)
 		{
 			FClothPainter* ClothPainter = static_cast<FClothPainter*>(PaintMode->GetMeshPainter());
@@ -201,7 +201,7 @@ void SClothPaintTab::OnAssetSelectionChanged(TWeakObjectPtr<UClothingAsset> InAs
 		}
 	}
 
-	if(UClothingAsset* Asset = InAssetPtr.Get())
+	if(UClothingAssetCommon* Asset = InAssetPtr.Get())
 	{
 		TArray<UObject*> Objects;
 
@@ -220,7 +220,7 @@ bool SClothPaintTab::IsAssetDetailsPanelEnabled()
 
 		if(SelectedObjects.Num() > 0)
 		{
-			return SelectedObjects[0].Get() != UClothingAsset::StaticClass()->GetDefaultObject();
+			return SelectedObjects[0].Get() != UClothingAssetCommon::StaticClass()->GetDefaultObject();
 		}
 	}
 
diff --git a/Engine/Source/Editor/MovieSceneTools/Private/TrackEditors/TransformTrackEditor.cpp b/Engine/Source/Editor/MovieSceneTools/Private/TrackEditors/TransformTrackEditor.cpp
--- a/Engine/Source/Editor/MovieSceneTools/Private/TrackEditors/TransformTrackEditor.cpp
+++ b/Engine/Source/Editor/MovieSceneTools/Private/TrackEditors/TransformTrackEditor.cpp
@@ -195,12 +195,12 @@ void F3DTransformTrackEditor::BuildTrackContextMenu( FMenuBuilder& MenuBuilder,
 	}
 	UMovieScene3DTransformTrack* TransformTrack = Cast<UMovieScene3DTransformTrack>( Track );
 	MenuBuilder.AddMenuEntry(
-		NSLOCTEXT( "Sequencer", "PasteMatineeMoveTrack", "Paste Matinee Move Track" ),
-		NSLOCTEXT( "Sequencer", "PasteMatineeMoveTrackTooltip", "Pastes keys from a Matinee move track into this track." ),
+		NSLOCTEXT("Sequencer", "PasteMatineeMoveTrack", "Paste Matinee Move Track"),
+		NSLOCTEXT("Sequencer", "PasteMatineeMoveTrackTooltip", "Pastes keys from a Matinee move track into this track."),
 		FSlateIcon(),
 		FUIAction(
-			FExecuteAction::CreateStatic( &CopyInterpMoveTrack, GetSequencer().ToSharedRef(), MoveTrack, TransformTrack ),
-			FCanExecuteAction::CreateStatic( &CanCopyInterpMoveTrack, MoveTrack, TransformTrack ) ) );
+			FExecuteAction::CreateStatic(&CopyInterpMoveTrack, GetSequencer().ToSharedRef(), MoveTrack, TransformTrack),
+			FCanExecuteAction::CreateStatic(&CanCopyInterpMoveTrack, MoveTrack, TransformTrack)));
 
 	//		FCanExecuteAction::CreateLambda( [=]()->bool { return MoveTrack != nullptr && MoveTrack->GetNumKeys() > 0 && TransformTrack != nullptr; } ) ) );
 
@@ -389,6 +389,35 @@ void F3DTransformTrackEditor::OnPostPropertyChanged(UObject* InObject, FProperty
 	}
 }
 
+bool F3DTransformTrackEditor::CanAddTransformKeysForSelectedObjects() const
+{
+	// WASD hotkeys to fly the viewport can conflict with hotkeys for setting keyframes (ie. s). 
+// If the viewport is moving, disregard setting keyframes.
+	for (FLevelEditorViewportClient* LevelVC : GEditor->GetLevelViewportClients())
+	{
+		if (LevelVC && LevelVC->IsMovingCamera())
+		{
+			return false;
+		}
+	}
+	TArray<UObject*> SelectedObjects;
+	for (FSelectedEditableComponentIterator It(GEditor->GetSelectedEditableComponentIterator()); It; ++It)
+	{
+		USceneComponent* SceneComponent = Cast<USceneComponent>(*It);
+		if (SceneComponent)
+		{
+			return true;
+		}
+	}
+
+	if (SelectedObjects.Num() == 0)
+	{
+		USelection* CurrentSelection = GEditor->GetSelectedActors();
+		CurrentSelection->GetSelectedObjects(AActor::StaticClass(), SelectedObjects);
+	}
+	return SelectedObjects.Num() > 0;
+}
+
 void F3DTransformTrackEditor::OnAddTransformKeysForSelectedObjects( EMovieSceneTransformChannel Channel )
 {
 	// WASD hotkeys to fly the viewport can conflict with hotkeys for setting keyframes (ie. s). 
@@ -431,20 +460,22 @@ void F3DTransformTrackEditor::BindCommands(TSharedRef<FUICommandList> SequencerC
 	CommandBindings = MakeShared<FUICommandList>();
 	CommandBindings->MapAction(
 		Commands.AddTransformKey,
-		FExecuteAction::CreateSP( this, &F3DTransformTrackEditor::OnAddTransformKeysForSelectedObjects, EMovieSceneTransformChannel::All ) );
+		FExecuteAction::CreateSP(this, &F3DTransformTrackEditor::OnAddTransformKeysForSelectedObjects, EMovieSceneTransformChannel::All),
+		FCanExecuteAction::CreateSP(this, &F3DTransformTrackEditor::CanAddTransformKeysForSelectedObjects));
+		
 
 	CommandBindings->MapAction(
 		Commands.AddTranslationKey,
-		FExecuteAction::CreateSP( this, &F3DTransformTrackEditor::OnAddTransformKeysForSelectedObjects, EMovieSceneTransformChannel::Translation ) );
-
+		FExecuteAction::CreateSP( this, &F3DTransformTrackEditor::OnAddTransformKeysForSelectedObjects, EMovieSceneTransformChannel::Translation ),
+		FCanExecuteAction::CreateSP(this, &F3DTransformTrackEditor::CanAddTransformKeysForSelectedObjects));
 	CommandBindings->MapAction(
 		Commands.AddRotationKey,
-		FExecuteAction::CreateSP( this, &F3DTransformTrackEditor::OnAddTransformKeysForSelectedObjects, EMovieSceneTransformChannel::Rotation ) );
-
+		FExecuteAction::CreateSP( this, &F3DTransformTrackEditor::OnAddTransformKeysForSelectedObjects, EMovieSceneTransformChannel::Rotation ),
+		FCanExecuteAction::CreateSP(this, &F3DTransformTrackEditor::CanAddTransformKeysForSelectedObjects));
 	CommandBindings->MapAction(
 		Commands.AddScaleKey,
-		FExecuteAction::CreateSP( this, &F3DTransformTrackEditor::OnAddTransformKeysForSelectedObjects, EMovieSceneTransformChannel::Scale ) );
-
+		FExecuteAction::CreateSP( this, &F3DTransformTrackEditor::OnAddTransformKeysForSelectedObjects, EMovieSceneTransformChannel::Scale ),
+		FCanExecuteAction::CreateSP(this, &F3DTransformTrackEditor::CanAddTransformKeysForSelectedObjects));
 	Commands.BindingCount++;
 
 	// Add these bindings to Sequencer
@@ -869,20 +900,20 @@ bool F3DTransformTrackEditor::ModifyGeneratedKeysByCurrentAndWeight(UObject *Obj
 		CurrentPos = Transform.Translation;
 		CurrentRot = Transform.Rotation;
 		CurrentScale = Transform.Scale;
-		break;
-	}
-	FMovieSceneChannelProxy& Proxy = SectionToKey->GetChannelProxy();
-	GeneratedTotalKeys[0]->ModifyByCurrentAndWeight(Proxy, KeyTime, (void *)&CurrentPos.X, Weight);
-	GeneratedTotalKeys[1]->ModifyByCurrentAndWeight(Proxy, KeyTime, (void *)&CurrentPos.Y, Weight);
-	GeneratedTotalKeys[2]->ModifyByCurrentAndWeight(Proxy, KeyTime, (void *)&CurrentPos.Z, Weight);
-	GeneratedTotalKeys[3]->ModifyByCurrentAndWeight(Proxy, KeyTime, (void *)&CurrentRot.Roll, Weight);
-	GeneratedTotalKeys[4]->ModifyByCurrentAndWeight(Proxy, KeyTime, (void *)&CurrentRot.Pitch, Weight);
-	GeneratedTotalKeys[5]->ModifyByCurrentAndWeight(Proxy, KeyTime, (void *)&CurrentRot.Yaw, Weight);
-	GeneratedTotalKeys[6]->ModifyByCurrentAndWeight(Proxy, KeyTime, (void *)&CurrentScale.X, Weight);
-	GeneratedTotalKeys[7]->ModifyByCurrentAndWeight(Proxy, KeyTime, (void *)&CurrentScale.Y, Weight);
-	GeneratedTotalKeys[8]->ModifyByCurrentAndWeight(Proxy, KeyTime, (void *)&CurrentScale.Z, Weight);
-	return true;
 
+		FMovieSceneChannelProxy& Proxy = SectionToKey->GetChannelProxy();
+		GeneratedTotalKeys[0]->ModifyByCurrentAndWeight(Proxy, KeyTime, (void *)&CurrentPos.X, Weight);
+		GeneratedTotalKeys[1]->ModifyByCurrentAndWeight(Proxy, KeyTime, (void *)&CurrentPos.Y, Weight);
+		GeneratedTotalKeys[2]->ModifyByCurrentAndWeight(Proxy, KeyTime, (void *)&CurrentPos.Z, Weight);
+		GeneratedTotalKeys[3]->ModifyByCurrentAndWeight(Proxy, KeyTime, (void *)&CurrentRot.Roll, Weight);
+		GeneratedTotalKeys[4]->ModifyByCurrentAndWeight(Proxy, KeyTime, (void *)&CurrentRot.Pitch, Weight);
+		GeneratedTotalKeys[5]->ModifyByCurrentAndWeight(Proxy, KeyTime, (void *)&CurrentRot.Yaw, Weight);
+		GeneratedTotalKeys[6]->ModifyByCurrentAndWeight(Proxy, KeyTime, (void *)&CurrentScale.X, Weight);
+		GeneratedTotalKeys[7]->ModifyByCurrentAndWeight(Proxy, KeyTime, (void *)&CurrentScale.Y, Weight);
+		GeneratedTotalKeys[8]->ModifyByCurrentAndWeight(Proxy, KeyTime, (void *)&CurrentScale.Z, Weight);
+		return true;
+	}
+	return false;
 }
 
 void AddUnwoundKey(FMovieSceneFloatChannel& Channel, FFrameNumber Time, float Value)
@@ -957,6 +988,10 @@ void F3DTransformTrackEditor::ImportAnimSequenceTransforms(const FAssetData& Ass
 		TransformTrack->Modify();
 
 		UMovieScene3DTransformSection* Section = Cast<UMovieScene3DTransformSection>(TransformTrack->CreateNewSection());
+
+		Section->SetBlendType(EMovieSceneBlendType::Additive);
+		Section->SetMask(EMovieSceneTransformChannel::Translation | EMovieSceneTransformChannel::Rotation);
+		
 		FFrameRate TickResolution = Section->GetTypedOuter<UMovieScene>()->GetTickResolution();
 
 		TArrayView<FMovieSceneFloatChannel*> FloatChannels = Section->GetChannelProxy().GetChannels<FMovieSceneFloatChannel>();
diff --git a/Engine/Source/Editor/Persona/Persona.Build.cs b/Engine/Source/Editor/Persona/Persona.Build.cs
--- a/Engine/Source/Editor/Persona/Persona.Build.cs
+++ b/Engine/Source/Editor/Persona/Persona.Build.cs
@@ -30,6 +30,7 @@ public class Persona : ModuleRules
                 "AnimationEditor",
                 "MeshReductionInterface",
                 "SequenceRecorder",
+                "AnimationBlueprintEditor",
             }
         );
 
@@ -55,9 +56,9 @@ public class Persona : ModuleRules
                 "RHI",
                 "Json",
                 "JsonUtilities",
-                "ClothingSystemRuntimeNv",
                 "ClothingSystemEditorInterface",
                 "ClothingSystemRuntimeInterface",
+                "ClothingSystemRuntimeCommon",
                 "AnimGraphRuntime",
                 "UnrealEd",
                 "CommonMenuExtensions",
diff --git a/Engine/Source/Editor/Persona/Private/AnimGraphNodeDetails.cpp b/Engine/Source/Editor/Persona/Private/AnimGraphNodeDetails.cpp
--- a/Engine/Source/Editor/Persona/Private/AnimGraphNodeDetails.cpp
+++ b/Engine/Source/Editor/Persona/Private/AnimGraphNodeDetails.cpp
@@ -44,6 +44,8 @@
 #include "Engine/SkeletalMeshSocket.h"
 #include "Styling/CoreStyle.h"
 #include "Widgets/Input/SCheckBox.h"
+#include "LODInfoUILayout.h"
+#include "IPersonaToolkit.h"
 
 #define LOCTEXT_NAMESPACE "KismetNodeWithOptionalPinsDetails"
 
@@ -169,8 +171,6 @@ void FAnimGraphNodeDetails::CustomizeDetails(class IDetailLayoutBuilder& DetailB
 
 						ShowHidePropertyHandle->MarkHiddenByCustomization();
 
-						TSharedRef<SWidget> ShowHidePropertyWidget = CreateAsPinWidget(ShowHidePropertyHandle);
-
 						ValueWidget->SetVisibility(TAttribute<EVisibility>::Create(TAttribute<EVisibility>::FGetter::CreateSP(this, &FAnimGraphNodeDetails::GetVisibilityOfProperty, ShowHidePropertyHandle)));
 
 						// If we have an edit condition, that comes as part of the default name widget, so just use a text block to avoid duplicate checkboxes
@@ -187,26 +187,7 @@ void FAnimGraphNodeDetails::CustomizeDetails(class IDetailLayoutBuilder& DetailB
 							PropertyNameWidget = NameWidget;
 						}
 
-						NameWidget = 
-							SNew(SHorizontalBox)
-							+SHorizontalBox::Slot()
-							.HAlign(HAlign_Left)
-							.VAlign(VAlign_Center)
-							.FillWidth(1)
-							[
-								SNew(SBox)
-								.Clipping(EWidgetClipping::ClipToBounds)
-								[
-									PropertyNameWidget.ToSharedRef()
-								]
-							]
-							+SHorizontalBox::Slot()
-							.AutoWidth()
-							.HAlign(HAlign_Right)
-							.VAlign(VAlign_Center)
-							[
-								ShowHidePropertyWidget
-							];
+						NameWidget = PropertyNameWidget;
 
 						// we only show children if visibility is one
 						// whenever toggles, this gets called, so it will be refreshed
@@ -250,68 +231,6 @@ void FAnimGraphNodeDetails::CustomizeDetails(class IDetailLayoutBuilder& DetailB
 	}
 }
 
-TSharedRef<SWidget> FAnimGraphNodeDetails::CreateAsPinWidget(TSharedRef<IPropertyHandle> InPropertyHandle)
-{
-	TWeakPtr<IPropertyHandle> WeakPropertyHandle = InPropertyHandle;
-
-	auto IsCheckedLambda = [WeakPropertyHandle]()
-	{
-		if(WeakPropertyHandle.IsValid())
-		{
-			bool bValue;
-			FPropertyAccess::Result Result = WeakPropertyHandle.Pin()->GetValue(bValue);
-			if(Result == FPropertyAccess::MultipleValues)
-			{
-				return ECheckBoxState::Undetermined;
-			}
-			else
-			{
-				return bValue ? ECheckBoxState::Checked : ECheckBoxState::Unchecked;
-			}
-		}
-
-		return ECheckBoxState::Unchecked;
-	};
-
-	auto OnCheckStateChangedLambda = [WeakPropertyHandle](ECheckBoxState InCheckBoxState)
-	{
-		if(WeakPropertyHandle.IsValid())
-		{
-			bool bValue = InCheckBoxState == ECheckBoxState::Checked;
-			WeakPropertyHandle.Pin()->SetValue(bValue);
-		}
-	};
-
-	auto ImageLambda = [WeakPropertyHandle]()
-	{
-		if(WeakPropertyHandle.IsValid())
-		{
-			bool bValue;
-			FPropertyAccess::Result Result = WeakPropertyHandle.Pin()->GetValue(bValue);
-			if(Result == FPropertyAccess::MultipleValues)
-			{
-				return FEditorStyle::GetBrush("Kismet.VariableList.HideForInstance");
-			}
-			else
-			{
-				return bValue ? FEditorStyle::GetBrush("Kismet.VariableList.ExposeForInstance") : FEditorStyle::GetBrush("Kismet.VariableList.HideForInstance");
-			}
-		}
-
-		return FEditorStyle::GetBrush("Kismet.VariableList.HideForInstance");
-	};
-
-	return SNew(SCheckBox)
-		.ToolTipText(LOCTEXT("AsPinTooltip", "Show/hide this property as a pin on the node"))
-		.IsChecked_Lambda(IsCheckedLambda)
-		.OnCheckStateChanged_Lambda(OnCheckStateChangedLambda)
-		.Style(FEditorStyle::Get(), "CheckboxLookToggleButtonCheckbox")
-		[
-			SNew(SImage)
-			.Image_Lambda(ImageLambda)
-			.ColorAndOpacity(FLinearColor::Black)
-		];
-}
 
 TSharedRef<SWidget> FAnimGraphNodeDetails::CreatePropertyWidget(UProperty* TargetProperty, TSharedRef<IPropertyHandle> TargetPropertyHandle, UClass* NodeClass)
 {
@@ -599,6 +518,7 @@ void FBoneReferenceCustomization::SetEditableSkeleton(TSharedRef<IPropertyHandle
 	UAnimGraphNode_Base* AnimGraphNode = nullptr;
 	USkeletalMesh* SkeletalMesh = nullptr;
 	UAnimationAsset * AnimationAsset = nullptr;
+	ULODInfoUILayout* LODInfoUILayout = nullptr;
 
 	USkeleton* TargetSkeleton = nullptr;
 	TSharedPtr<IEditableSkeleton> EditableSkeleton;
@@ -617,6 +537,16 @@ void FBoneReferenceCustomization::SetEditableSkeleton(TSharedRef<IPropertyHandle
 			TargetSkeleton = SkeletalMesh->Skeleton;
 			break;
 		}
+
+		LODInfoUILayout = Cast<ULODInfoUILayout>(*OuterIter);
+		if (LODInfoUILayout)
+		{
+			SkeletalMesh = LODInfoUILayout->GetPersonaToolkit()->GetPreviewMesh();
+			check(SkeletalMesh);
+			TargetSkeleton = SkeletalMesh->Skeleton;
+			break;
+		}
+
 		AnimationAsset = Cast<UAnimationAsset>(*OuterIter);
 		if (AnimationAsset)
 		{
@@ -711,10 +641,10 @@ TSharedRef<IPropertyTypeCustomization> FBoneSocketTargetCustomization::MakeInsta
 void FBoneSocketTargetCustomization::CustomizeChildren(TSharedRef<IPropertyHandle> StructPropertyHandle, class IDetailChildrenBuilder& ChildBuilder, IPropertyTypeCustomizationUtils& StructCustomizationUtils)
 {
 	// set property handle 
- 	SetPropertyHandle(StructPropertyHandle);
- 	// set editable skeleton info from struct
- 	SetEditableSkeleton(StructPropertyHandle);
- 	Build(StructPropertyHandle, ChildBuilder);
+	SetPropertyHandle(StructPropertyHandle);
+	// set editable skeleton info from struct
+	SetEditableSkeleton(StructPropertyHandle);
+	Build(StructPropertyHandle, ChildBuilder);
 }
 
 void FBoneSocketTargetCustomization::SetPropertyHandle(TSharedRef<IPropertyHandle> StructPropertyHandle)
@@ -1225,4 +1155,135 @@ void FPlayerTreeViewEntry::GenerateNameWidget(TSharedPtr<SHorizontalBox> Box)
 		];
 }
 
+void FAnimGraphNodeShowAsPinExtension::GetOptionalPinData(const IPropertyHandle& PropertyHandle, int32& OutOptionalPinIndex, UAnimGraphNode_Base*& OutAnimGraphNode) const
+{
+	OutOptionalPinIndex = INDEX_NONE;
+
+	TArray<UObject*> Objects;
+	PropertyHandle.GetOuterObjects(Objects);
+
+	UProperty* Property = PropertyHandle.GetProperty();
+	if(Property)
+	{
+		OutAnimGraphNode = Cast<UAnimGraphNode_Base>(Objects[0]);
+		if (OutAnimGraphNode != nullptr)
+		{
+			OutOptionalPinIndex = OutAnimGraphNode->ShowPinForProperties.IndexOfByPredicate([Property](const FOptionalPinFromProperty& InOptionalPin)
+			{
+				return Property->GetFName() == InOptionalPin.PropertyName;
+			});
+		}
+	}
+}
+
+bool FAnimGraphNodeShowAsPinExtension::IsPropertyExtendable(const UClass* InObjectClass, const IPropertyHandle& PropertyHandle) const
+{
+	int32 OptionalPinIndex;
+	UAnimGraphNode_Base* AnimGraphNode;
+	GetOptionalPinData(PropertyHandle, OptionalPinIndex, AnimGraphNode);
+
+	if(OptionalPinIndex != INDEX_NONE)
+	{
+		const FOptionalPinFromProperty& OptionalPin = AnimGraphNode->ShowPinForProperties[OptionalPinIndex];
+
+		// Not optional
+		if (!OptionalPin.bCanToggleVisibility && OptionalPin.bShowPin)
+		{
+			return false;
+		}
+
+		if(!PropertyHandle.GetProperty())
+		{
+			return false;
+		}
+
+		return OptionalPin.bCanToggleVisibility;
+	}
+
+	return false;
+}
+
+class SShowAsWidget : public SCompoundWidget
+{
+	SLATE_BEGIN_ARGS(SShowAsWidget) {}
+
+	SLATE_END_ARGS()
+
+	void Construct(const FArguments& InArgs, const TSharedRef<IPropertyHandle>& InPropertyHandle)
+	{
+		PropertyHandle = InPropertyHandle;
+
+		TSharedRef<SHorizontalBox> HorizontalBox = 
+			SNew(SHorizontalBox)
+			.ToolTipText(LOCTEXT("AsPinTooltip", "Show/hide this property as a pin on the node"));
+
+		TWeakPtr<SWidget> WeakHorizontalBox = HorizontalBox;
+
+		HorizontalBox->AddSlot()
+			.AutoWidth()
+			.VAlign(VAlign_Center)
+			[
+				SNew(STextBlock)
+				.Text(LOCTEXT("ExposeAsPinLabel", "Expose"))
+				.Font(IDetailLayoutBuilder::GetDetailFont())
+				.Visibility_Lambda([WeakHorizontalBox](){ return WeakHorizontalBox.IsValid() && WeakHorizontalBox.Pin()->IsHovered() ? EVisibility::Visible : EVisibility::Collapsed; })
+			];
+
+		HorizontalBox->AddSlot()
+			.AutoWidth()
+			.VAlign(VAlign_Center)
+			.Padding(3.0f, 0.0f, 0.0f, 0.0f)
+			[
+				SNew(SCheckBox)
+				.IsChecked(this, &SShowAsWidget::IsChecked)
+				.OnCheckStateChanged(this, &SShowAsWidget::OnCheckStateChanged)
+			];
+
+		ChildSlot
+		[
+			HorizontalBox
+		];
+	}
+
+	ECheckBoxState IsChecked() const
+	{
+		bool bValue;
+		FPropertyAccess::Result Result = PropertyHandle->GetValue(bValue);
+		if(Result == FPropertyAccess::MultipleValues)
+		{
+			return ECheckBoxState::Undetermined;
+		}
+		else
+		{
+			return bValue ? ECheckBoxState::Checked : ECheckBoxState::Unchecked;
+		}
+
+		return ECheckBoxState::Unchecked;
+	}
+
+	void OnCheckStateChanged(ECheckBoxState InCheckBoxState)
+	{
+		bool bValue = InCheckBoxState == ECheckBoxState::Checked;
+		PropertyHandle->SetValue(bValue);
+	}
+
+	TSharedPtr<IPropertyHandle> PropertyHandle;
+};
+
+TSharedRef<SWidget> FAnimGraphNodeShowAsPinExtension::GenerateExtensionWidget(const IDetailLayoutBuilder& InDetailBuilder, const UClass* InObjectClass, TSharedPtr<IPropertyHandle> PropertyHandle)
+{
+	int32 OptionalPinIndex;
+	UAnimGraphNode_Base* AnimGraphNode;
+	GetOptionalPinData(*PropertyHandle.Get(), OptionalPinIndex, AnimGraphNode);
+	check(OptionalPinIndex != INDEX_NONE);
+
+	const FName OptionalPinArrayEntryName(*FString::Printf(TEXT("ShowPinForProperties[%d].bShowPin"), OptionalPinIndex));
+	TSharedRef<IPropertyHandle> ShowHidePropertyHandle = InDetailBuilder.GetProperty(OptionalPinArrayEntryName, UAnimGraphNode_Base::StaticClass());
+
+	ShowHidePropertyHandle->MarkHiddenByCustomization();
+
+	return SNew(SShowAsWidget, ShowHidePropertyHandle);
+}
+
 #undef LOCTEXT_NAMESPACE
+
diff --git a/Engine/Source/Editor/SkeletalMeshEditor/Private/SkeletalMeshEditor.cpp b/Engine/Source/Editor/SkeletalMeshEditor/Private/SkeletalMeshEditor.cpp
--- a/Engine/Source/Editor/SkeletalMeshEditor/Private/SkeletalMeshEditor.cpp
+++ b/Engine/Source/Editor/SkeletalMeshEditor/Private/SkeletalMeshEditor.cpp
@@ -23,7 +23,7 @@
 #include "Rendering/SkeletalMeshModel.h"
 
 #include "Animation/DebugSkelMeshComponent.h"
-#include "Assets/ClothingAsset.h"
+#include "ClothingAsset.h"
 #include "SCreateClothingSettingsPanel.h"
 #include "ClothingSystemEditorInterfaceModule.h"
 #include "Preferences/PersonaOptions.h"
@@ -44,6 +44,8 @@
 #include "EditorFramework/AssetImportData.h"
 #include "Factories/FbxSkeletalMeshImportData.h"
 
+#include "Misc/MessageDialog.h"
+
 const FName SkeletalMeshEditorAppIdentifier = FName(TEXT("SkeletalMeshEditorApp"));
 
 const FName SkeletalMeshEditorModes::SkeletalMeshEditorMode(TEXT("SkeletalMeshEditorMode"));
@@ -78,6 +80,112 @@ FSkeletalMeshEditor::~FSkeletalMeshEditor()
 	}
 }
 
+bool IsReductionParentBaseLODUseSkeletalMeshBuildWorkflow(USkeletalMesh* SkeletalMesh, int32 TestLODIndex)
+{
+	FSkeletalMeshLODInfo* LODInfo = SkeletalMesh->GetLODInfo(TestLODIndex);
+	if (LODInfo == nullptr || !SkeletalMesh->GetImportedModel() || !SkeletalMesh->GetImportedModel()->LODModels.IsValidIndex(TestLODIndex))
+	{
+		return false;
+	}
+
+	if (SkeletalMesh->GetImportedModel()->LODModels[TestLODIndex].RawSkeletalMeshBulkData.IsBuildDataAvailable())
+	{
+		return true;
+	}
+	if (LODInfo->bHasBeenSimplified || SkeletalMesh->IsReductionActive(TestLODIndex))
+	{
+		int32 ReduceBaseLOD = LODInfo->ReductionSettings.BaseLOD;
+		if (ReduceBaseLOD < TestLODIndex)
+		{
+			return IsReductionParentBaseLODUseSkeletalMeshBuildWorkflow(SkeletalMesh, ReduceBaseLOD);
+		}
+	}
+	return false;
+}
+
+bool FSkeletalMeshEditor::OnRequestClose()
+{
+	bool bAllowClose = true;
+
+	if (PersonaToolkit.IsValid() && SkeletalMesh)
+	{
+		bool bHaveModifiedLOD = false;
+		for (int32 LODIndex = 0; LODIndex < SkeletalMesh->GetLODNum(); ++LODIndex)
+		{
+			FSkeletalMeshLODInfo* LODInfo = SkeletalMesh->GetLODInfo(LODIndex);
+			if (LODInfo == nullptr || !SkeletalMesh->GetImportedModel() || !SkeletalMesh->GetImportedModel()->LODModels.IsValidIndex(LODIndex))
+			{
+				continue;
+			}
+			
+			//Do not prevent exiting if we are not using the skeletal mesh build workflow
+			if (!SkeletalMesh->GetImportedModel()->LODModels[LODIndex].RawSkeletalMeshBulkData.IsBuildDataAvailable())
+			{
+				if (!LODInfo->bHasBeenSimplified && !SkeletalMesh->IsReductionActive(LODIndex))
+				{
+					continue;
+				}
+				//Do not prevent exit if the generated LOD is base on a LODModel not using the skeletalmesh build workflow
+				int32 ReduceBaseLOD = LODInfo->ReductionSettings.BaseLOD;
+				if(!IsReductionParentBaseLODUseSkeletalMeshBuildWorkflow(SkeletalMesh, ReduceBaseLOD))
+				{
+					continue;
+				}
+				
+			}
+
+			bool bValidLODSettings = false;
+			if (SkeletalMesh->LODSettings != nullptr)
+			{
+				const int32 NumSettings = FMath::Min(SkeletalMesh->LODSettings->GetNumberOfSettings(), SkeletalMesh->GetLODNum());
+				if (LODIndex < NumSettings)
+				{
+					bValidLODSettings = true;
+				}
+			}
+			const FSkeletalMeshLODGroupSettings* SkeletalMeshLODGroupSettings = bValidLODSettings ? &SkeletalMesh->LODSettings->GetSettingsForLODLevel(LODIndex) : nullptr;
+
+			FGuid BuildGUID = LODInfo->ComputeDeriveDataCacheKey(SkeletalMeshLODGroupSettings);
+			if (LODInfo->BuildGUID != BuildGUID)
+			{
+				bHaveModifiedLOD = true;
+				break;
+			}
+			FString BuildStringID = SkeletalMesh->GetImportedModel()->LODModels[LODIndex].GetLODModelDeriveDataKey();
+			if (SkeletalMesh->GetImportedModel()->LODModels[LODIndex].BuildStringID != BuildStringID)
+			{
+				bHaveModifiedLOD = true;
+				break;
+			}
+		}
+
+		if (bHaveModifiedLOD)
+		{
+			// find out the user wants to do with this dirty material
+			EAppReturnType::Type OkCancelReply = FMessageDialog::Open(
+				EAppMsgType::OkCancel,
+				FText::Format(LOCTEXT("SkeletalMeshEditorShouldApplyLODChanges", "We have to apply level of detail changes to {0} before exiting the skeletal mesh editor."), FText::FromString(PersonaToolkit->GetMesh()->GetName()))
+			);
+
+			switch (OkCancelReply)
+			{
+			case EAppReturnType::Ok:
+				{
+					SkeletalMesh->MarkPackageDirty();
+					SkeletalMesh->PostEditChange();
+					bAllowClose = true;
+				}
+				break;
+			case EAppReturnType::Cancel:
+				// Don't exit.
+				bAllowClose = false;
+				break;
+			}
+		}
+	}
+	return bAllowClose;
+}
+
 void FSkeletalMeshEditor::RegisterTabSpawners(const TSharedRef<class FTabManager>& InTabManager)
 {
 	WorkspaceMenuCategory = InTabManager->AddLocalWorkspaceMenuCategory(LOCTEXT("WorkspaceMenu_SkeletalMeshEditor", "Skeletal Mesh Editor"));
@@ -546,8 +654,7 @@ void FSkeletalMeshEditor::OnRemoveSectionFromLodAndBelowMenuItemClicked(int32 Lo
 		const FSkeletalMeshLODInfo* CurrentSkeletalMeshLODInfo = SkeletalMesh->GetLODInfo(GenerateLodIndex);
 		if (CurrentSkeletalMeshLODInfo != nullptr && CurrentSkeletalMeshLODInfo->bHasBeenSimplified && BaseLodIndexes.Contains(CurrentSkeletalMeshLODInfo->ReductionSettings.BaseLOD))
 		{
-			bool bRestoreClothing = true;
-			FLODUtilities::SimplifySkeletalMeshLOD(UpdateContext, GenerateLodIndex, true, bRestoreClothing);
+			FLODUtilities::SimplifySkeletalMeshLOD(UpdateContext, GenerateLodIndex, true);
 			BaseLodIndexes.Add(GenerateLodIndex);
 		}
 	}
@@ -569,12 +676,12 @@ void FSkeletalMeshEditor::FillApplyClothingAssetMenu(FMenuBuilder& MenuBuilder,
 	{
 		for(UClothingAssetBase* BaseAsset : Mesh->MeshClothingAssets)
 		{
-			UClothingAsset* ClothAsset = CastChecked<UClothingAsset>(BaseAsset);
+			UClothingAssetCommon* ClothAsset = CastChecked<UClothingAssetCommon>(BaseAsset);
 
 			FUIAction Action;
 			Action.CanExecuteAction = FCanExecuteAction::CreateSP(this, &FSkeletalMeshEditor::CanApplyClothing, InLodIndex, InSectionIndex);
 
-			const int32 NumClothLods = ClothAsset->LodData.Num();
+			const int32 NumClothLods = ClothAsset->GetNumLods();
 			for(int32 ClothLodIndex = 0; ClothLodIndex < NumClothLods; ++ClothLodIndex)
 			{
 				Action.ExecuteAction = FExecuteAction::CreateSP(this, &FSkeletalMeshEditor::OnApplyClothingAssetClicked, BaseAsset, InLodIndex, InSectionIndex, ClothLodIndex);
@@ -668,7 +775,7 @@ void FSkeletalMeshEditor::OnCreateClothingAssetMenuItemClicked(FSkeletalMeshClot
 						{
 							SectionIndex = i;
 							AssetLodIndex = SkelLod.Sections[i].ClothingData.AssetLodIndex;
-							TargetAssetPtr->UnbindFromSkeletalMesh(Mesh, Params.TargetLod);
+							RemoveClothing(Params.TargetLod, SectionIndex);
 							break;
 						}
 					}
@@ -786,13 +893,54 @@ void FSkeletalMeshEditor::ApplyClothing(UClothingAssetBase* InAsset, int32 InLod
 {
 	USkeletalMesh* Mesh = GetPersonaToolkit()->GetPreviewMesh();
 
-	if(UClothingAsset* ClothingAsset = Cast<UClothingAsset>(InAsset))
+	if (Mesh == nullptr || Mesh->GetImportedModel() == nullptr || !Mesh->GetImportedModel()->LODModels.IsValidIndex(InLodIndex))
 	{
-		ClothingAsset->BindToSkeletalMesh(Mesh, InLodIndex, InSectionIndex, InClothingLod);
+		return;
 	}
-	else
+
+	FSkeletalMeshLODModel& LODModel = Mesh->GetImportedModel()->LODModels[InLodIndex];
+	const FSkelMeshSection& Section = LODModel.Sections[InSectionIndex];
+	FScopedSuspendAlternateSkinWeightPreview ScopedSuspendAlternateSkinnWeightPreview(Mesh);
 	{
-		RemoveClothing(InLodIndex, InSectionIndex);
+		FScopedSkeletalMeshPostEditChange ScopedPostEditChange(Mesh);
+		FScopedTransaction Transaction(LOCTEXT("SkeletalMeshEditorApplyClothingTransaction", "Persona editor: Apply Section Cloth"));
+		Mesh->Modify();
+
+		FSkelMeshSourceSectionUserData& OriginalSectionData = LODModel.UserSectionsData.FindOrAdd(Section.OriginalDataSectionIndex);
+		auto ClearOriginalSectionUserData = [&OriginalSectionData]()
+		{
+			OriginalSectionData.CorrespondClothAssetIndex = INDEX_NONE;
+			OriginalSectionData.ClothingData.AssetGuid = FGuid();
+			OriginalSectionData.ClothingData.AssetLodIndex = INDEX_NONE;
+		};
+		if (UClothingAssetCommon* ClothingAsset = Cast<UClothingAssetCommon>(InAsset))
+		{
+			// Look for a currently bound asset an unbind it if necessary first
+			if (UClothingAssetBase* CurrentAsset = Mesh->GetSectionClothingAsset(InLodIndex, InSectionIndex))
+			{
+				CurrentAsset->UnbindFromSkeletalMesh(Mesh, InLodIndex);
+				ClearOriginalSectionUserData();
+			}
+
+			if (ClothingAsset->BindToSkeletalMesh(Mesh, InLodIndex, InSectionIndex, InClothingLod))
+			{
+				//Successful bind so set the SectionUserData
+				int32 AssetIndex = INDEX_NONE;
+				check(Mesh->MeshClothingAssets.Find(ClothingAsset, AssetIndex));
+				OriginalSectionData.CorrespondClothAssetIndex = AssetIndex;
+				OriginalSectionData.ClothingData.AssetGuid = ClothingAsset->GetAssetGuid();
+				OriginalSectionData.ClothingData.AssetLodIndex = InClothingLod;
+			}
+		}
+		else if (Mesh)
+		{
+			//User set none, so unbind anything that is bind
+			if (UClothingAssetBase* CurrentAsset = Mesh->GetSectionClothingAsset(InLodIndex, InSectionIndex))
+			{
+				CurrentAsset->UnbindFromSkeletalMesh(Mesh, InLodIndex);
+				ClearOriginalSectionUserData();
+			}
+		}
 	}
 }
 
@@ -800,11 +948,31 @@ void FSkeletalMeshEditor::RemoveClothing(int32 InLodIndex, int32 InSectionIndex)
 {
 	USkeletalMesh* Mesh = GetPersonaToolkit()->GetPreviewMesh();
 
-	if(Mesh)
+	if (Mesh->GetImportedModel() == nullptr || !Mesh->GetImportedModel()->LODModels.IsValidIndex(InLodIndex))
 	{
-		if(UClothingAssetBase* CurrentAsset = Mesh->GetSectionClothingAsset(InLodIndex, InSectionIndex))
+		return;
+	}
+
+	FSkeletalMeshLODModel& LODModel = Mesh->GetImportedModel()->LODModels[InLodIndex];
+	const FSkelMeshSection& Section = LODModel.Sections[InSectionIndex];
+	FScopedSuspendAlternateSkinWeightPreview ScopedSuspendAlternateSkinnWeightPreview(Mesh);
+	{
+		FScopedSkeletalMeshPostEditChange ScopedPostEditChange(Mesh);
+		FScopedTransaction Transaction(LOCTEXT("SkeletalMeshEditorRemoveClothingTransaction", "Persona editor: Remove Section Cloth"));
+		Mesh->Modify();
+
+		FSkelMeshSourceSectionUserData& OriginalSectionData = LODModel.UserSectionsData.FindOrAdd(Section.OriginalDataSectionIndex);
+		auto ClearOriginalSectionUserData = [&OriginalSectionData]()
+		{
+			OriginalSectionData.CorrespondClothAssetIndex = INDEX_NONE;
+			OriginalSectionData.ClothingData.AssetGuid = FGuid();
+			OriginalSectionData.ClothingData.AssetLodIndex = INDEX_NONE;
+		};
+		// Look for a currently bound asset an unbind it if necessary first
+		if (UClothingAssetBase* CurrentAsset = Mesh->GetSectionClothingAsset(InLodIndex, InSectionIndex))
 		{
 			CurrentAsset->UnbindFromSkeletalMesh(Mesh, InLodIndex);
+			ClearOriginalSectionUserData();
 		}
 	}
 }
@@ -883,92 +1051,92 @@ UObject* FSkeletalMeshEditor::HandleGetAsset()
 
 bool FSkeletalMeshEditor::HandleReimportMeshInternal(int32 SourceFileIndex /*= INDEX_NONE*/, bool bWithNewFile /*= false*/)
 {
-	// Make sure we clear any skin weight profile currently being previewed
-	UDebugSkelMeshComponent* MeshComponent = GetPersonaToolkit()->GetPreviewScene()->GetPreviewMeshComponent();
-	const FName ProfileName = MeshComponent ? MeshComponent->GetCurrentSkinWeightProfileName() : NAME_None;
-	if (MeshComponent)
-	{
-		MeshComponent->ClearSkinWeightProfile();
-	}
-
-	// Reimport the asset
-	const bool bResult = FReimportManager::Instance()->Reimport(SkeletalMesh, true, true, TEXT(""), nullptr, SourceFileIndex, bWithNewFile);
-
-	// Refresh skeleton tree
-	SkeletonTree->Refresh();
-
-	// Re-set skin weigh profile
-	if (MeshComponent && (ProfileName != NAME_None))
+	FScopedSuspendAlternateSkinWeightPreview ScopedSuspendAlternateSkinnWeightPreview(SkeletalMesh);
+	bool bResult = false;
 	{
-		MeshComponent->SetSkinWeightProfile(ProfileName);
+		FScopedSkeletalMeshPostEditChange ScopedPostEditChange(SkeletalMesh);
+		// Reimport the asset
+		bResult = FReimportManager::Instance()->Reimport(SkeletalMesh, true, true, TEXT(""), nullptr, SourceFileIndex, bWithNewFile);
+		// Refresh skeleton tree
+		SkeletonTree->Refresh();
 	}
-
 	return bResult;
 }
 
 void FSkeletalMeshEditor::HandleReimportMesh(int32 SourceFileIndex /*= INDEX_NONE*/)
 {
+	FScopedSuspendAlternateSkinWeightPreview ScopedSuspendAlternateSkinnWeightPreview(SkeletalMesh);
+
 	HandleReimportMeshInternal(SourceFileIndex, false);
 }
 
 void FSkeletalMeshEditor::HandleReimportMeshWithNewFile(int32 SourceFileIndex /*= INDEX_NONE*/)
 {
+	FScopedSuspendAlternateSkinWeightPreview ScopedSuspendAlternateSkinnWeightPreview(SkeletalMesh);
+
 	HandleReimportMeshInternal(SourceFileIndex, true);
 }
 
 void ReimportAllCustomLODs(USkeletalMesh* SkeletalMesh, UDebugSkelMeshComponent* PreviewMeshComponent, bool bWithNewFile)
 {
-	//Find the dependencies of the generated LOD
-	TArray<bool> Dependencies;
-	Dependencies.AddZeroed(SkeletalMesh->GetLODNum());
-	//Avoid making LOD 0 to true in the dependencies since everything that should be regenerate base on LOD 0 is already regenerate at this point.
-	//But we need to regenerate every generated LOD base on any re-import custom LOD
-	//Reimport all custom LODs
-	for (int32 LodIndex = 1; LodIndex < SkeletalMesh->GetLODNum(); ++LodIndex)
+	FScopedSuspendAlternateSkinWeightPreview ScopedSuspendAlternateSkinnWeightPreview(SkeletalMesh);
 	{
-		//Do not reimport LOD that was re-import with the base mesh
-		if (SkeletalMesh->GetLODInfo(LodIndex)->bImportWithBaseMesh)
-		{
-			continue;
-		}
-		if (SkeletalMesh->GetLODInfo(LodIndex)->bHasBeenSimplified == false)
+		FScopedSkeletalMeshPostEditChange ScopedPostEditChange(SkeletalMesh);
+
+		//Find the dependencies of the generated LOD
+		TArray<bool> Dependencies;
+		Dependencies.AddZeroed(SkeletalMesh->GetLODNum());
+		//Avoid making LOD 0 to true in the dependencies since everything that should be regenerate base on LOD 0 is already regenerate at this point.
+		//But we need to regenerate every generated LOD base on any re-import custom LOD
+		//Reimport all custom LODs
+		for (int32 LodIndex = 1; LodIndex < SkeletalMesh->GetLODNum(); ++LodIndex)
 		{
-			FString SourceFilenameBackup = SkeletalMesh->GetLODInfo(LodIndex)->SourceImportFilename;
-			if (bWithNewFile)
+			//Do not reimport LOD that was re-import with the base mesh
+			if (SkeletalMesh->GetLODInfo(LodIndex)->bImportWithBaseMesh)
 			{
-				SkeletalMesh->GetLODInfo(LodIndex)->SourceImportFilename.Empty();
+				continue;
 			}
-
-			if (!FbxMeshUtils::ImportMeshLODDialog(SkeletalMesh, LodIndex))
+			if (SkeletalMesh->GetLODInfo(LodIndex)->bHasBeenSimplified == false)
 			{
+				FString SourceFilenameBackup = SkeletalMesh->GetLODInfo(LodIndex)->SourceImportFilename;
 				if (bWithNewFile)
 				{
-					SkeletalMesh->GetLODInfo(LodIndex)->SourceImportFilename = SourceFilenameBackup;
+					SkeletalMesh->GetLODInfo(LodIndex)->SourceImportFilename.Empty();
+				}
+
+				if (!FbxMeshUtils::ImportMeshLODDialog(SkeletalMesh, LodIndex, false))
+				{
+					if (bWithNewFile)
+					{
+						SkeletalMesh->GetLODInfo(LodIndex)->SourceImportFilename = SourceFilenameBackup;
+					}
+				}
+				else
+				{
+					Dependencies[LodIndex] = true;
 				}
 			}
-			else
+			else if (Dependencies[SkeletalMesh->GetLODInfo(LodIndex)->ReductionSettings.BaseLOD])
 			{
+				//Regenerate the LOD
+				FSkeletalMeshUpdateContext UpdateContext;
+				UpdateContext.SkeletalMesh = SkeletalMesh;
+				UpdateContext.AssociatedComponents.Push(PreviewMeshComponent);
+				FLODUtilities::SimplifySkeletalMeshLOD(UpdateContext, LodIndex, false);
 				Dependencies[LodIndex] = true;
 			}
 		}
-		else if(Dependencies[SkeletalMesh->GetLODInfo(LodIndex)->ReductionSettings.BaseLOD])
-		{
-			//Regenerate the LOD
-			FSkeletalMeshUpdateContext UpdateContext;
-			UpdateContext.SkeletalMesh = SkeletalMesh;
-			UpdateContext.AssociatedComponents.Push(PreviewMeshComponent);
-			bool bRestoreClothing = true;
-			FLODUtilities::SimplifySkeletalMeshLOD(UpdateContext, LodIndex, true, bRestoreClothing);
-			Dependencies[LodIndex] = true;
-		}
 	}
 }
 
 void FSkeletalMeshEditor::HandleReimportAllMesh(int32 SourceFileIndex /*= INDEX_NONE*/)
 {
+	FScopedSuspendAlternateSkinWeightPreview ScopedSuspendAlternateSkinnWeightPreview(SkeletalMesh);
 	// Reimport the asset
 	if (SkeletalMesh)
 	{
+		FScopedSkeletalMeshPostEditChange ScopedPostEditChange(SkeletalMesh);
+
 		//Reimport base LOD
 		if (HandleReimportMeshInternal(SourceFileIndex, false))
 		{
@@ -980,9 +1148,12 @@ void FSkeletalMeshEditor::HandleReimportAllMesh(int32 SourceFileIndex /*= INDEX_
 
 void FSkeletalMeshEditor::HandleReimportAllMeshWithNewFile(int32 SourceFileIndex /*= INDEX_NONE*/)
 {
+	FScopedSuspendAlternateSkinWeightPreview ScopedSuspendAlternateSkinnWeightPreview(SkeletalMesh);
+
 	// Reimport the asset
 	if (SkeletalMesh)
 	{
+		FScopedSkeletalMeshPostEditChange ScopedPostEditChange(SkeletalMesh);
 		TArray<UObject*> ImportObjs;
 		ImportObjs.Add(SkeletalMesh);
 		if (HandleReimportMeshInternal(SourceFileIndex, true))
diff --git a/Engine/Source/Editor/UnrealEd/Private/Animation/DebugSkelMeshComponent.cpp b/Engine/Source/Editor/UnrealEd/Private/Animation/DebugSkelMeshComponent.cpp
--- a/Engine/Source/Editor/UnrealEd/Private/Animation/DebugSkelMeshComponent.cpp
+++ b/Engine/Source/Editor/UnrealEd/Private/Animation/DebugSkelMeshComponent.cpp
@@ -17,7 +17,7 @@
 #include "Rendering/SkeletalMeshRenderData.h"
 #include "Rendering/SkeletalMeshModel.h"
 
-#include "Assets/ClothingAsset.h"
+#include "ClothingAsset.h"
 #include "ClothingSimulation.h"
 #include "Utils/ClothingMeshUtils.h"
 #include "DynamicMeshBuilder.h"
@@ -740,17 +740,15 @@ void UDebugSkelMeshComponent::RebuildClothingSectionsFixedVerts()
 
 				if(BaseAsset)
 				{
-					UClothingAsset* ConcreteAsset = Cast<UClothingAsset>(BaseAsset);
-					FClothLODData& LodData = ConcreteAsset->LodData[Section.ClothingData.AssetLodIndex];
+					UClothingAssetCommon* ConcreteAsset = Cast<UClothingAssetCommon>(BaseAsset);
+					UClothLODDataBase* LodData = ConcreteAsset->ClothLodData[Section.ClothingData.AssetLodIndex];
 
 					for(FMeshToMeshVertData& VertData : Section.ClothMappingData)
 					{
-						float TriangleDistanceMax = 0.0f;
-						TriangleDistanceMax += LodData.PhysicalMeshData.MaxDistances[VertData.SourceMeshVertIndices[0]];
-						TriangleDistanceMax += LodData.PhysicalMeshData.MaxDistances[VertData.SourceMeshVertIndices[1]];
-						TriangleDistanceMax += LodData.PhysicalMeshData.MaxDistances[VertData.SourceMeshVertIndices[2]];
-
-						if(TriangleDistanceMax == 0.0f)
+						if(LodData->PhysicalMeshData->IsFixed(
+							VertData.SourceMeshVertIndices[0],
+							VertData.SourceMeshVertIndices[1],
+							VertData.SourceMeshVertIndices[2]))
 						{
 							VertData.SourceMeshVertIndices[3] = 0xFFFF;
 						}
@@ -819,9 +817,9 @@ void UDebugSkelMeshComponent::RefreshSelectedClothingSkinnedPositions()
 
 		if(Asset)
 		{
-			UClothingAsset* ConcreteAsset = Cast<UClothingAsset>(*Asset);
+			UClothingAssetCommon* ConcreteAsset = Cast<UClothingAssetCommon>(*Asset);
 
-			if(ConcreteAsset->LodData.IsValidIndex(SelectedClothingLodForPainting))
+			if(ConcreteAsset->ClothLodData.IsValidIndex(SelectedClothingLodForPainting))
 			{
 				SkinnedSelectedClothingPositions.Reset();
 				SkinnedSelectedClothingNormals.Reset();
@@ -830,9 +828,9 @@ void UDebugSkelMeshComponent::RefreshSelectedClothingSkinnedPositions()
 				// Pass LOD0 to collect all bones
 				GetCurrentRefToLocalMatrices(RefToLocals, 0);
 
-				FClothLODData& LodData = ConcreteAsset->LodData[SelectedClothingLodForPainting];
+				UClothLODDataBase* LodData = ConcreteAsset->ClothLodData[SelectedClothingLodForPainting];
 
-				ClothingMeshUtils::SkinPhysicsMesh(ConcreteAsset->UsedBoneIndices, LodData.PhysicalMeshData, FTransform::Identity, RefToLocals.GetData(), RefToLocals.Num(), SkinnedSelectedClothingPositions, SkinnedSelectedClothingNormals);
+				ClothingMeshUtils::SkinPhysicsMesh(ConcreteAsset->UsedBoneIndices, *LodData->PhysicalMeshData, FTransform::Identity, RefToLocals.GetData(), RefToLocals.Num(), SkinnedSelectedClothingPositions, SkinnedSelectedClothingNormals);
 				RebuildCachedClothBounds();
 			}
 		}
@@ -1033,17 +1031,17 @@ FDebugSkelMeshDynamicData::FDebugSkelMeshDynamicData(UDebugSkelMeshComponent* In
 				{
 					ClothingSimDataIndexWhenPainting = ClothingAssetIndex;
 
-					if(UClothingAsset* ConcreteAsset = Cast<UClothingAsset>(BaseAsset))
+					if(UClothingAssetCommon* ConcreteAsset = Cast<UClothingAssetCommon>(BaseAsset))
 					{
-						if(ConcreteAsset->LodData.IsValidIndex(InComponent->SelectedClothingLodForPainting))
+						if(ConcreteAsset->ClothLodData.IsValidIndex(InComponent->SelectedClothingLodForPainting))
 						{
-							FClothLODData& LodData = ConcreteAsset->LodData[InComponent->SelectedClothingLodForPainting];
+							UClothLODDataBase* LodData = ConcreteAsset->ClothLodData[InComponent->SelectedClothingLodForPainting];
 
-							ClothingSimIndices = LodData.PhysicalMeshData.Indices;
+							ClothingSimIndices = LodData->PhysicalMeshData->Indices;
 
-							if(LodData.ParameterMasks.IsValidIndex(InComponent->SelectedClothingLodMaskForPainting))
+							if(LodData->ParameterMasks.IsValidIndex(InComponent->SelectedClothingLodMaskForPainting))
 							{
-								FClothParameterMask_PhysMesh& Mask = LodData.ParameterMasks[InComponent->SelectedClothingLodMaskForPainting];
+								FPointWeightMap& Mask = LodData->ParameterMasks[InComponent->SelectedClothingLodMaskForPainting];
 
 								ClothingVisiblePropertyValues = Mask.GetValueArray();
 							}
@@ -1057,3 +1055,37 @@ FDebugSkelMeshDynamicData::FDebugSkelMeshDynamicData(UDebugSkelMeshComponent* In
 	}
 }
 
+FScopedSuspendAlternateSkinWeightPreview::FScopedSuspendAlternateSkinWeightPreview(USkeletalMesh* SkeletalMesh)
+{
+	SuspendedComponentArray.Empty(2);
+	if (SkeletalMesh != nullptr)
+	{
+		// Now iterate over all skeletal mesh components and unregister them from the world, we will reregister them in the destructor
+		for (TObjectIterator<UDebugSkelMeshComponent> It; It; ++It)
+		{
+			UDebugSkelMeshComponent* DebugSKComp = *It;
+			if (DebugSKComp->SkeletalMesh == SkeletalMesh)
+			{
+				const FName ProfileName = DebugSKComp->GetCurrentSkinWeightProfileName();
+				if (ProfileName != NAME_None)
+				{
+					DebugSKComp->ClearSkinWeightProfile();
+					TTuple<UDebugSkelMeshComponent*, FName> ComponentTupple;
+					ComponentTupple.Key = DebugSKComp;
+					ComponentTupple.Value = ProfileName;
+					SuspendedComponentArray.Add(ComponentTupple);
+				}
+			}
+		}
+	}
+}
+
+FScopedSuspendAlternateSkinWeightPreview::~FScopedSuspendAlternateSkinWeightPreview()
+{
+	//Put back the skin weight profile for all editor debug component
+	for (const TTuple<UDebugSkelMeshComponent*, FName>& ComponentTupple : SuspendedComponentArray)
+	{
+		ComponentTupple.Key->SetSkinWeightProfile(ComponentTupple.Value);
+	}
+	SuspendedComponentArray.Empty();
+}
