Fix DDC bug when using the FLODUtilities::RegenerateLOD to force the regeneration of LODs.

diff --git a/Engine/Source/Developer/SkeletalMeshUtilitiesCommon/Private/LODUtilities.cpp b/Engine/Source/Developer/SkeletalMeshUtilitiesCommon/Private/LODUtilities.cpp
--- a/Engine/Source/Developer/SkeletalMeshUtilitiesCommon/Private/LODUtilities.cpp
+++ b/Engine/Source/Developer/SkeletalMeshUtilitiesCommon/Private/LODUtilities.cpp
@@ -221,6 +221,9 @@ bool FLODUtilities::RegenerateLOD(USkeletalMesh* SkeletalMesh, int32 NewLODCount
 		FSkeletalMeshUpdateContext UpdateContext;
 		UpdateContext.SkeletalMesh = SkeletalMesh;
 
+		//If we force a regenerate, we want to invalidate the DCC so the render data get rebuilded
+		SkeletalMesh->InvalidateDeriveDataCacheGUID();
+
 		// remove LODs
 		int32 CurrentNumLODs = SkeletalMesh->GetLODNum();
 		if (LODCount < CurrentNumLODs)
