Fix alternate skinning

diff --git a/Engine/Source/Developer/MeshUtilities/Private/MeshUtilities.cpp b/Engine/Source/Developer/MeshUtilities/Private/MeshUtilities.cpp
--- a/Engine/Source/Developer/MeshUtilities/Private/MeshUtilities.cpp
+++ b/Engine/Source/Developer/MeshUtilities/Private/MeshUtilities.cpp
@@ -5749,6 +5749,8 @@ void FMeshUtilities::GenerateRuntimeSkinWeightData(const FSkeletalMeshLODModel*
 		const int32 NumInfluences = TargetLODModel.GetMaxBoneInfluences();
 		InOutSkinWeightOverrideData.NumWeightsPerVertex = NumInfluences;
 
+		const bool b16BitBoneIndices = TargetLODModel.DoSectionsUse16BitBoneIndex();
+		InOutSkinWeightOverrideData.b16BitBoneIndices = b16BitBoneIndices;
 
 		TArray<FRawSkinWeight> UniqueWeights;
 		for (int32 VertexIndex = 0; VertexIndex < TargetVertices.Num(); ++VertexIndex)
@@ -5793,8 +5795,18 @@ void FMeshUtilities::GenerateRuntimeSkinWeightData(const FSkeletalMeshLODModel*
 					{
 						const FBoneIndexType Index = SourceSkinWeight.InfluenceBones[InfluenceIndex];
 						const uint8 Weight = SourceSkinWeight.InfluenceWeights[InfluenceIndex];
+						
+						if (b16BitBoneIndices)
+						{
+							InOutSkinWeightOverrideData.BoneIDs.AddZeroed(2);
+							FBoneIndexType* BoneIndex = (FBoneIndexType*)&InOutSkinWeightOverrideData.BoneIDs[InOutSkinWeightOverrideData.BoneIDs.Num() - 2];
+							*BoneIndex = Index;
+						}
+						else
+						{
+							InOutSkinWeightOverrideData.BoneIDs.Add((uint8)Index);
+						}
 
-						InOutSkinWeightOverrideData.BoneIDs.Add(Index);
 						InOutSkinWeightOverrideData.BoneWeights.Add(Weight);
 					}
 
diff --git a/Engine/Source/Developer/SkeletalMeshUtilitiesCommon/Private/LODUtilities.cpp b/Engine/Source/Developer/SkeletalMeshUtilitiesCommon/Private/LODUtilities.cpp
--- a/Engine/Source/Developer/SkeletalMeshUtilitiesCommon/Private/LODUtilities.cpp
+++ b/Engine/Source/Developer/SkeletalMeshUtilitiesCommon/Private/LODUtilities.cpp
@@ -14,6 +14,7 @@
 #include "Engine/SkeletalMesh.h"
 #include "EditorFramework/AssetImportData.h"
 #include "MeshUtilities.h"
+#include "MeshUtilitiesCommon.h"
 #include "ClothingAsset.h"
 #include "OverlappingCorners.h"
 #include "Framework/Commands/UIAction.h"
@@ -1320,10 +1321,10 @@ namespace VertexMatchNameSpace
 struct FTriangleOctreeSemantics
 {
 	// When a leaf gets more than this number of elements, it will split itself into a node with multiple child leaves
-	enum { MaxElementsPerLeaf = 6 };
+	enum { MaxElementsPerLeaf = 10 };
 
 	// This is used for incremental updates.  When removing a polygon, larger values will cause leaves to be removed and collapsed into a parent node.
-	enum { MinInclusiveElementsPerNode = 7 };
+	enum { MinInclusiveElementsPerNode = 5 };
 
 	// How deep the tree can go.
 	enum { MaxNodeDepth = 20 };
@@ -1440,7 +1441,8 @@ void MatchVertexIndexUsingPosition(
 
 		//Use the OcTree to find closest triangle
 		FVector Extent(DistanceThreshold, DistanceThreshold, DistanceThreshold);
-		FBox CurBox(PositionSrc - Extent, PositionSrc + Extent);
+		FBoxCenterAndExtent CurBox(PositionSrc, Extent);
+		
 		while (OcTreeTriangleResults.Num() <= 0)
 		{
 			OcTree.FindElementsWithBoundsTest(CurBox, [&OcTreeTriangleResults](const FTriangleElement& Element)
@@ -1451,6 +1453,11 @@ void MatchVertexIndexUsingPosition(
 
 			//Increase the extend so we try to found in a larger area
 			Extent *= 2;
+			if (Extent.SizeSquared() >= BaseMeshPositionBound.GetSize().SizeSquared())
+			{
+				//Extend must not be bigger then the whole mesh, its acceptable to have error at this point
+				break;
+			}
 			CurBox = FBox(PositionSrc - Extent, PositionSrc + Extent);
 		}
 
@@ -1587,6 +1594,40 @@ bool FLODUtilities::UpdateAlternateSkinWeights(FSkeletalMeshLODModel& LODModelDe
 		UE_LOG(LogLODUtilities, Error, TEXT("Failed to import Skin Weight Profile as the target skeletal mesh (%s) or imported file does not contain UV coordinates."), *SkeletalMeshName);
 		return false;
 	}
+	
+	// Create a list of vertex Z/index pairs
+	TArray<FIndexAndZ> VertIndexAndZ;
+	VertIndexAndZ.Reserve(VertexNumberDest);
+	for (int32 VertexIndex = 0; VertexIndex < VertexNumberDest; ++VertexIndex)
+	{
+		new(VertIndexAndZ)FIndexAndZ(VertexIndex, ImportDataDest.Points[VertexIndex]);
+	}
+	// Sort the vertices by z value
+	VertIndexAndZ.Sort(FCompareIndexAndZ());
+	
+	auto FindSimilarPosition = [&VertIndexAndZ, &ImportDataDest](const FVector& Position, TArray<int32>& PositionMatches, const float ComparisonThreshold)
+	{
+		PositionMatches.Reset();
+		FIndexAndZ PositionZ = FIndexAndZ(0, Position);
+		// Search for duplicates, quickly!
+		for (int32 i = 0; i < VertIndexAndZ.Num(); i++)
+		{
+			if (PositionZ.Z - ComparisonThreshold > VertIndexAndZ[i].Z)
+			{
+				continue;
+			}
+			else if (PositionZ.Z + ComparisonThreshold < VertIndexAndZ[i].Z)
+			{
+				break;
+			}
+
+			const FVector& PositionA = ImportDataDest.Points[VertIndexAndZ[i].Index];
+			if (PointsEqual(PositionA, Position, ComparisonThreshold))
+			{
+				PositionMatches.Add(VertIndexAndZ[i].Index);
+			}
+		}
+	};
 
 	//Create a map linking all similar Position of destination vertex index
 	TMap<FVector, TArray<uint32>> PositionToVertexIndexDest;
@@ -1611,8 +1652,10 @@ bool FLODUtilities::UpdateAlternateSkinWeights(FSkeletalMeshLODModel& LODModelDe
 	{
 		const FVector& PositionSrc = ImportDataSrc.Points[VertexIndexSrc];
 		
-		TArray<uint32>* SimilarDestinationVertex = PositionToVertexIndexDest.Find(PositionSrc);
-		if (!SimilarDestinationVertex)
+		TArray<int32> SimilarDestinationVertex;
+		FindSimilarPosition(PositionSrc, SimilarDestinationVertex, KINDA_SMALL_NUMBER);
+
+		if (SimilarDestinationVertex.Num() == 0)
 		{
 			//Match with UV projection
 			VertexIndexToMatchWithUVs.Add(VertexIndexSrc);
@@ -1621,9 +1664,9 @@ bool FLODUtilities::UpdateAlternateSkinWeights(FSkeletalMeshLODModel& LODModelDe
 		{
 			//We have a direct match
 			VertexMatchNameSpace::FVertexMatchResult& VertexMatchDest = VertexIndexSrcToVertexIndexDestMatches.Add(VertexIndexSrc);
-			for (int32 MatchDestinationIndex = 0; MatchDestinationIndex < SimilarDestinationVertex->Num(); ++MatchDestinationIndex)
+			for (int32 MatchDestinationIndex = 0; MatchDestinationIndex < SimilarDestinationVertex.Num(); ++MatchDestinationIndex)
 			{
-				VertexMatchDest.VertexIndexes.Add((*SimilarDestinationVertex)[MatchDestinationIndex]);
+				VertexMatchDest.VertexIndexes.Add(SimilarDestinationVertex[MatchDestinationIndex]);
 				VertexMatchDest.Ratios.Add(1.0f);
 			}
 		}
@@ -1648,6 +1691,7 @@ bool FLODUtilities::UpdateAlternateSkinWeights(FSkeletalMeshLODModel& LODModelDe
 				continue;
 			}
 		}
+		bAllSourceVertexAreMatch = VertexIndexSrcToVertexIndexDestMatches.Num() == PointNumberSrc;
 	}
 	
 	
@@ -1675,11 +1719,11 @@ bool FLODUtilities::UpdateAlternateSkinWeights(FSkeletalMeshLODModel& LODModelDe
 	for (int32 InfluenceIndexSrc = 0; InfluenceIndexSrc < InfluenceNumberSrc; ++InfluenceIndexSrc)
 	{
 		const SkeletalMeshImportData::FRawBoneInfluence& InfluenceSrc = ImportDataSrc.Influences[InfluenceIndexSrc];
-		uint32 VertexIndexSource = InfluenceSrc.VertexIndex;
-		uint32 BoneIndexSource = InfluenceSrc.BoneIndex;
+		int32 VertexIndexSource = InfluenceSrc.VertexIndex;
+		int32 BoneIndexSource = InfluenceSrc.BoneIndex;
 		float Weight = InfluenceSrc.Weight;
 		//We need to remap the source bone index to have the matching target bone index
-		uint32 BoneIndexDest = RemapBoneIndexSrcToDest[BoneIndexSource];
+		int32 BoneIndexDest = RemapBoneIndexSrcToDest[BoneIndexSource];
 		if (BoneIndexDest != INDEX_NONE)
 		{
 			//Find the match destination vertex index
@@ -1804,9 +1848,9 @@ bool FLODUtilities::UpdateAlternateSkinWeights(FSkeletalMeshLODModel& LODModelDe
 	//
 	//////////////////////////////////////////////////////////////////////////
 
+	bool bBuildSuccess = true;
 	//Prepare the build data to rebuild the asset with the alternate influences
 	//The chunking can be different when we have alternate influences
-
 	//Grab the build data from ImportDataDest
 	TArray<FVector> LODPointsDest;
 	TArray<SkeletalMeshImportData::FMeshWedge> LODWedgesDest;
@@ -1829,7 +1873,7 @@ bool FLODUtilities::UpdateAlternateSkinWeights(FSkeletalMeshLODModel& LODModelDe
 	TArray<FText> WarningMessages;
 	TArray<FName> WarningNames;
 	//Build the destination mesh with the Alternate influences, so the chunking is done properly.
-	bool bBuildSuccess = MeshUtilities.BuildSkeletalMesh(LODModelDest, RefSkeleton, LODInfluencesDest, LODWedgesDest, LODFacesDest, LODPointsDest, LODPointToRawMapDest, BuildOptions, &WarningMessages, &WarningNames);
+	bBuildSuccess = MeshUtilities.BuildSkeletalMesh(LODModelDest, RefSkeleton, LODInfluencesDest, LODWedgesDest, LODFacesDest, LODPointsDest, LODPointToRawMapDest, BuildOptions, &WarningMessages, &WarningNames);
 	//Re-Apply the user section changes, the UserSectionsData is map to original section and should match the builded LODModel
 	LODModelDest.SyncronizeUserSectionsDataArray();
 
@@ -1881,8 +1925,12 @@ bool FLODUtilities::UpdateAlternateSkinWeights(USkeletalMesh* SkeletalMeshDest,
 	//Resave the bulk data with the new or refreshed data
 	SkeletalMeshDest->SaveLODImportedData(LODIndexDest, ImportDataDest);
 
-	//Build the alternate buffer with all the data into the bulk
-	return UpdateAlternateSkinWeights(SkeletalMeshDest, ProfileNameDest, LODIndexDest, OverlappingThresholds, ShouldImportNormals, ShouldImportTangents, bUseMikkTSpace, bComputeWeightedNormals);
+	if(!SkeletalMeshDest->IsLODImportedDataBuildAvailable(LODIndexDest))
+	{
+		//Build the alternate buffer with all the data into the bulk, in case the build data is not existing (old asset)
+		return UpdateAlternateSkinWeights(SkeletalMeshDest, ProfileNameDest, LODIndexDest, OverlappingThresholds, ShouldImportNormals, ShouldImportTangents, bUseMikkTSpace, bComputeWeightedNormals);
+	}
+	return true;
 }
 
 void FLODUtilities::GenerateImportedSkinWeightProfileData(const FSkeletalMeshLODModel& LODModelDest, FImportedSkinWeightProfileData &ImportedProfileData)
diff --git a/Engine/Source/Editor/Persona/Private/SkinWeightProfileHelpers.cpp b/Engine/Source/Editor/Persona/Private/SkinWeightProfileHelpers.cpp
--- a/Engine/Source/Editor/Persona/Private/SkinWeightProfileHelpers.cpp
+++ b/Engine/Source/Editor/Persona/Private/SkinWeightProfileHelpers.cpp
@@ -86,8 +86,11 @@ void FSkinWeightProfileHelpers::ImportSkinWeightProfile(USkeletalMesh* InSkeleta
 				NotificationInfo.ExpireDuration = 2.5f;
 				FSlateNotificationManager::Get().AddNotification(NotificationInfo);
 			}
-			
-			FLODUtilities::RegenerateDependentLODs(InSkeletalMesh, 0);
+
+			if(!InSkeletalMesh->IsLODImportedDataBuildAvailable(ImportSettings->LODIndex))
+			{
+				FLODUtilities::RegenerateDependentLODs(InSkeletalMesh, ImportSettings->LODIndex);
+			}
 		}
 	}
 }
@@ -104,7 +107,10 @@ void FSkinWeightProfileHelpers::ImportSkinWeightProfileLOD(USkeletalMesh* InSkel
 		const bool bResult = FSkinWeightsUtilities::ImportAlternateSkinWeight(InSkeletalMesh, PickedFileName, LODIndex, ProfileName);
 		if (bResult)
 		{
-			FLODUtilities::RegenerateDependentLODs(InSkeletalMesh, LODIndex);
+			if(!InSkeletalMesh->IsLODImportedDataBuildAvailable(LODIndex))
+			{
+				FLODUtilities::RegenerateDependentLODs(InSkeletalMesh, LODIndex);
+			}
 
 			FNotificationInfo NotificationInfo(FText::GetEmpty());
 			NotificationInfo.Text = FText::Format(LOCTEXT("LODImportSuccessful", "Skin Weights for LOD {0} imported successfully!"), FText::AsNumber(LODIndex));
@@ -152,8 +158,11 @@ void FSkinWeightProfileHelpers::ReimportSkinWeightProfileLOD(USkeletalMesh* InSk
 
 			if (bResult)
 			{
-				// Make sure we regenerate any LOD data that is based off the now re imported LOD
-				FLODUtilities::RegenerateDependentLODs(InSkeletalMesh, LODIndex);
+				if (!InSkeletalMesh->IsLODImportedDataBuildAvailable(LODIndex))
+				{
+					// Make sure we regenerate any LOD data that is based off the now re imported LOD
+					FLODUtilities::RegenerateDependentLODs(InSkeletalMesh, LODIndex);
+				}
 
 				FNotificationInfo NotificationInfo(FText::GetEmpty());
 				NotificationInfo.Text = FText::Format(LOCTEXT("ReimportedSkinWeightsForLOD", "Reimported Skin Weights for LOD {0} succesfully!"), FText::AsNumber(LODIndex));
@@ -195,9 +204,11 @@ void FSkinWeightProfileHelpers::RemoveSkinWeightProfile(USkeletalMesh* InSkeleta
 			{
 				// Delete the alternate influences data
 				LODModel.SkinWeightProfiles.Remove(InProfileName);
-				// Regenerate this generated LOD
-				FLODUtilities::SimplifySkeletalMeshLOD(UpdateContext, LODIndex);
-				
+				if (!InSkeletalMesh->IsLODImportedDataBuildAvailable(LODIndex))
+				{
+					// Regenerate this generated LOD
+					FLODUtilities::SimplifySkeletalMeshLOD(UpdateContext, LODIndex);
+				}
 				// Goto next LOD
 				continue;
 			}
@@ -238,8 +249,11 @@ void FSkinWeightProfileHelpers::RemoveSkinWeightProfileLOD(USkeletalMesh* InSkel
 		{
 			FSkinWeightProfileHelpers::ClearSkinWeightProfileInstanceOverrides(InSkeletalMesh, InProfileName);
 
-			//Regenerate dependent LODs if we remove the LOD successfully
-			FLODUtilities::RegenerateDependentLODs(InSkeletalMesh, LODIndex);
+			if (!InSkeletalMesh->IsLODImportedDataBuildAvailable(LODIndex))
+			{
+				//Regenerate dependent LODs if we remove the LOD successfully
+				FLODUtilities::RegenerateDependentLODs(InSkeletalMesh, LODIndex);
+			}
 
 			FNotificationInfo NotificationInfo(FText::GetEmpty());
 			NotificationInfo.Text = FText::Format(LOCTEXT("RemovedLODWeights", "Removed Skin Weights for LOD {0}!"), FText::AsNumber(LODIndex));
diff --git a/Engine/Source/Editor/UnrealEd/Private/SkinWeightsUtilities.cpp b/Engine/Source/Editor/UnrealEd/Private/SkinWeightsUtilities.cpp
--- a/Engine/Source/Editor/UnrealEd/Private/SkinWeightsUtilities.cpp
+++ b/Engine/Source/Editor/UnrealEd/Private/SkinWeightsUtilities.cpp
@@ -43,7 +43,7 @@ bool FSkinWeightsUtilities::ImportAlternateSkinWeight(USkeletalMesh* SkeletalMes
 	
 	if (LODInfo && LODInfo->bHasBeenSimplified && LODInfo->ReductionSettings.BaseLOD != TargetLODIndex)
 	{
-		//We cannot remove alternate skin weights profile for a generated LOD
+		//We cannot import alternate skin weights profile for a generated LOD
 		UE_LOG(LogSkinWeightsUtilities, Error, TEXT("Cannot import Skin Weight Profile for a generated LOD."));
 		return false;
 	}
@@ -98,7 +98,7 @@ bool FSkinWeightsUtilities::ImportAlternateSkinWeight(USkeletalMesh* SkeletalMes
 	//Force some material options
 	if (FbxFactory->ImportUI->TextureImportData)
 	{
-		FbxFactory->ImportUI->TextureImportData->MaterialSearchLocation = EMaterialSearchLocation::Local;
+		FbxFactory->ImportUI->TextureImportData->MaterialSearchLocation = EMaterialSearchLocation::DoNotSearch;
 		FbxFactory->ImportUI->TextureImportData->BaseMaterialName.Reset();
 	}
 
@@ -129,7 +129,7 @@ bool FSkinWeightsUtilities::ImportAlternateSkinWeight(USkeletalMesh* SkeletalMes
 	Task->DestinationPath = ImportAssetPath;
 	Task->bSave = false;
 	Task->DestinationName = FGuid::NewGuid().ToString(EGuidFormats::Digits);
-	Task->Options = FbxFactory->ImportUI->SkeletalMeshImportData;
+	Task->Options = FbxFactory->ImportUI;
 	Task->Filename = AbsoluteFilePath;
 	Task->Factory = FbxFactory;
 	FbxFactory->SetAssetImportTask(Task);
@@ -203,7 +203,6 @@ bool FSkinWeightsUtilities::ImportAlternateSkinWeight(USkeletalMesh* SkeletalMes
 				bool ShouldImportTangents = ImportOptions.ShouldImportTangents();
 				bool bUseMikkTSpace = ImportOptions.NormalGenerationMethod == EFBXNormalGenerationMethod::MikkTSpace;
 
-				TArray<FRawSkinWeight>& SkinWeights = ProfileData.SkinWeights;
 				bResult = FLODUtilities::UpdateAlternateSkinWeights(SkeletalMesh, ProfileName, TmpSkeletalMesh, TargetLODIndex, SrcLodIndex, OverlappingThresholds, ShouldImportNormals, ShouldImportTangents, bUseMikkTSpace, ImportOptions.bComputeWeightedNormals);
 				
 				if (!bResult)
diff --git a/Engine/Source/Runtime/Engine/Private/Animation/SkinWeightProfile.cpp b/Engine/Source/Runtime/Engine/Private/Animation/SkinWeightProfile.cpp
--- a/Engine/Source/Runtime/Engine/Private/Animation/SkinWeightProfile.cpp
+++ b/Engine/Source/Runtime/Engine/Private/Animation/SkinWeightProfile.cpp
@@ -764,7 +764,7 @@ void FRuntimeSkinWeightProfileData::ApplyOverrides(FSkinWeightVertexBuffer* Over
 				const uint32 InfluenceOffset = VertexIndexOverridePair.Value;
 				
 				const uint8* BoneData = TargetSkinWeightData + (VertexIndex * VertexStride);
-				const uint8* WeightData = BoneData + WeightDataOffset;				
+				const uint8* WeightData = BoneData + WeightDataOffset;
 
 #if !UE_BUILD_SHIPPING
 				uint32 VertexOffset = 0;
@@ -772,8 +772,10 @@ void FRuntimeSkinWeightProfileData::ApplyOverrides(FSkinWeightVertexBuffer* Over
 				OverrideBuffer->GetVertexInfluenceOffsetCount(VertexIndex, VertexOffset, VertexInfluenceCount);
 				check(NumWeightsPerVertex <= VertexInfluenceCount);
 				check((void*)(((uint8*)TargetSkinWeightData) + VertexOffset) == (void*)BoneData);
-#endif				
-				FMemory::Memcpy((void*)BoneData, &BoneIDs[InfluenceOffset * NumWeightsPerVertex], sizeof(FBoneIndexType) * NumWeightsPerVertex);
+				check(b16BitBoneIndices == OverrideBuffer->Use16BitBoneIndex());
+#endif
+				// BoneIDs either contains FBoneIndexType entries spanning (2) uint8 values, or single uint8 bone indices (1)
+				FMemory::Memcpy((void*)BoneData, &BoneIDs[InfluenceOffset * NumWeightsPerVertex * OverrideBuffer->GetBoneIndexByteSize()], OverrideBuffer->GetBoneIndexByteSize() * NumWeightsPerVertex);
 				FMemory::Memcpy((void*)WeightData, &BoneWeights[InfluenceOffset * NumWeightsPerVertex], sizeof(uint8) * NumWeightsPerVertex);
 			}
 		}
@@ -806,9 +808,10 @@ void FRuntimeSkinWeightProfileData::ApplyDefaultOverride(FSkinWeightVertexBuffer
 				Buffer->GetVertexInfluenceOffsetCount(VertexIndex, VertexOffset, VertexInfluenceCount);
 				check(NumWeightsPerVertex <= VertexInfluenceCount);
 				check((void*)(((uint8*)TargetSkinWeightData) + VertexOffset) == (void*)BoneData);
+				check(b16BitBoneIndices == Buffer->Use16BitBoneIndex());
 #endif
-				
-				FMemory::Memcpy((void*)BoneData, &BoneIDs[InfluenceOffset * NumWeightsPerVertex], sizeof(FBoneIndexType) * NumWeightsPerVertex);
+				// BoneIDs either contains FBoneIndexType entries spanning (2) uint8 values, or single uint8 bone indices (1)
+				FMemory::Memcpy((void*)BoneData, &BoneIDs[InfluenceOffset * NumWeightsPerVertex * Buffer->GetBoneIndexByteSize()], Buffer->GetBoneIndexByteSize() * NumWeightsPerVertex);
 				FMemory::Memcpy((void*)WeightData, &BoneWeights[InfluenceOffset * NumWeightsPerVertex], sizeof(uint8) * NumWeightsPerVertex);
 			}
 		}
diff --git a/Engine/Source/Runtime/Engine/Private/SkeletalMeshRenderData.cpp b/Engine/Source/Runtime/Engine/Private/SkeletalMeshRenderData.cpp
--- a/Engine/Source/Runtime/Engine/Private/SkeletalMeshRenderData.cpp
+++ b/Engine/Source/Runtime/Engine/Private/SkeletalMeshRenderData.cpp
@@ -59,7 +59,7 @@ static void SerializeLODInfoForDDC(USkeletalMesh* SkeletalMesh, FString& KeySuff
 // differences, etc.) replace the version GUID below with a new one.
 // In case of merge conflicts with DDC versions, you MUST generate a new GUID
 // and set this new GUID as the version.                                       
-#define SKELETALMESH_DERIVEDDATA_VER TEXT("8972F79737FD445387F6914AC5D6374D")
+#define SKELETALMESH_DERIVEDDATA_VER TEXT("EACF3F0E29C74570AA25B6AF9D9E957B")
 
 static const FString& GetSkeletalMeshDerivedDataVersion()
 {
diff --git a/Engine/Source/Runtime/Engine/Public/Animation/SkinWeightProfile.h b/Engine/Source/Runtime/Engine/Public/Animation/SkinWeightProfile.h
--- a/Engine/Source/Runtime/Engine/Public/Animation/SkinWeightProfile.h
+++ b/Engine/Source/Runtime/Engine/Public/Animation/SkinWeightProfile.h
@@ -101,12 +101,14 @@ struct FRuntimeSkinWeightProfileData
 	TArray<uint16> Weights_DEPRECATED;	
 #endif 
 
-	TArray<FBoneIndexType> BoneIDs;
+	// Either contains FBoneIndexType or uint8 bone indices
+	TArray<uint8> BoneIDs;
 	TArray<uint8> BoneWeights;
 	/** Map between Vertex Indices and the influence offset into BoneIDs/BoneWeights (DEPRECATED and entries of OverridesInfo) */
 	TMap<uint32, uint32> VertexIndexToInfluenceOffset;
 
 	uint8 NumWeightsPerVertex;
+	bool b16BitBoneIndices;
 	
 	friend FArchive& operator<<(FArchive& Ar, FRuntimeSkinWeightProfileData& OverrideData);
 };
diff --git a/Engine/Source/Runtime/Engine/Public/Rendering/SkeletalMeshLODModel.h b/Engine/Source/Runtime/Engine/Public/Rendering/SkeletalMeshLODModel.h
--- a/Engine/Source/Runtime/Engine/Public/Rendering/SkeletalMeshLODModel.h
+++ b/Engine/Source/Runtime/Engine/Public/Rendering/SkeletalMeshLODModel.h
@@ -449,7 +449,7 @@ public:
 	ENGINE_API void GetNonClothVertices(TArray<FSoftSkinVertex>& OutVertices) const;
 
 	ENGINE_API int32 GetMaxBoneInfluences() const;
-	bool DoSectionsUse16BitBoneIndex() const;
+	ENGINE_API bool DoSectionsUse16BitBoneIndex() const;
 
 	void GetResourceSizeEx(FResourceSizeEx& CumulativeResourceSize) const;
 
