To protect the skeletalmesh when using async operations like build or re-import we deprecate all public member variables of USkeletalMesh class and create the necessary getter and setter. Those member

diff --git a/Engine/Source/Developer/MeshBoneReduction/Private/MeshBoneReduction.cpp b/Engine/Source/Developer/MeshBoneReduction/Private/MeshBoneReduction.cpp
--- a/Engine/Source/Developer/MeshBoneReduction/Private/MeshBoneReduction.cpp
+++ b/Engine/Source/Developer/MeshBoneReduction/Private/MeshBoneReduction.cpp
@@ -60,7 +60,7 @@ public:
 			return false;
 		}
 
-		const TArray<FMeshBoneInfo> & RefBoneInfo = SkeletalMesh->RefSkeleton.GetRawRefBoneInfo();
+		const TArray<FMeshBoneInfo> & RefBoneInfo = SkeletalMesh->GetRefSkeleton().GetRawRefBoneInfo();
 		TArray<FBoneIndexType> BoneIndicesToRemove;
 
 		// originally this code was accumulating from LOD 0->DesiredLOd, but that should be done outside of tool if they want to
@@ -92,7 +92,7 @@ public:
 			{
 				if (BonesToRemoveSetting[Index] != NAME_None)
 				{
-					int32 BoneIndex = SkeletalMesh->RefSkeleton.FindRawBoneIndex(BonesToRemoveSetting[Index]);
+					int32 BoneIndex = SkeletalMesh->GetRefSkeleton().FindRawBoneIndex(BonesToRemoveSetting[Index]);
 
 					// we don't allow root to be removed
 					if (BoneIndex > 0)
@@ -278,10 +278,10 @@ public:
 
 		// Retrieve all bone names in skeleton
 		TArray<FName> BoneNames;
-		const int32 NumBones = SkeletalMesh->RefSkeleton.GetRawBoneNum();
+		const int32 NumBones = SkeletalMesh->GetRefSkeleton().GetRawBoneNum();
 		for (int32 BoneIndex = 0; BoneIndex < NumBones; ++BoneIndex)
 		{
-			BoneNames.Add(SkeletalMesh->RefSkeleton.GetBoneName(BoneIndex));
+			BoneNames.Add(SkeletalMesh->GetRefSkeleton().GetBoneName(BoneIndex));
 		}
 				
 		TArray<FMatrix> MultipliedBonePoses;
@@ -294,7 +294,7 @@ public:
 			MultipliedBonePoses.AddDefaulted(BonePoses.Num());
 
 			// Retrieve ref pose bone transforms
-			TArray<FTransform> RefBonePoses = SkeletalMesh->RefSkeleton.GetRawRefBonePose();
+			TArray<FTransform> RefBonePoses = SkeletalMesh->GetRefSkeleton().GetRawRefBonePose();
 			TArray<FMatrix> MultipliedRefBonePoses;
 			MultipliedRefBonePoses.AddDefaulted(RefBonePoses.Num());
 
@@ -305,8 +305,8 @@ public:
 			// Multiply out parent to child transforms for ref-pose
 			for (int32 BonePoseIndex = 0; BonePoseIndex < RefBonePoses.Num(); BonePoseIndex++)
 			{
-				const int32 BoneIndex = SkeletalMesh->RefSkeleton.FindRawBoneIndex(BoneNames[BonePoseIndex]);
-				const int32 ParentIndex = SkeletalMesh->RefSkeleton.GetParentIndex(BoneIndex);
+				const int32 BoneIndex = SkeletalMesh->GetRefSkeleton().FindRawBoneIndex(BoneNames[BonePoseIndex]);
+				const int32 ParentIndex = SkeletalMesh->GetRefSkeleton().GetParentIndex(BoneIndex);
 				MultipliedRefBonePoses[BoneIndex] = RefBonePoses[BoneIndex].ToMatrixWithScale();
 
 				if (ParentIndex != INDEX_NONE)
@@ -325,8 +325,8 @@ public:
 			// Multiply out parent to child transforms for bake-pose
 			for (int32 BonePoseIndex = 0; BonePoseIndex < BonePoses.Num(); BonePoseIndex++)
 			{
-				const int32 BoneIndex = SkeletalMesh->RefSkeleton.FindRawBoneIndex(BoneNames[BonePoseIndex]);
-				const int32 ParentIndex = SkeletalMesh->RefSkeleton.GetParentIndex(BoneIndex);
+				const int32 BoneIndex = SkeletalMesh->GetRefSkeleton().FindRawBoneIndex(BoneNames[BonePoseIndex]);
+				const int32 ParentIndex = SkeletalMesh->GetRefSkeleton().GetParentIndex(BoneIndex);
 				MultipliedBonePoses[BoneIndex] = BonePoses[BoneIndex].ToMatrixWithScale();
 				if (ParentIndex != INDEX_NONE)
 				{
@@ -362,7 +362,7 @@ public:
 	bool ReduceBoneCounts(USkeletalMesh* SkeletalMesh, int32 DesiredLOD, const TArray<FName>* BoneNamesToRemove, bool bCallPostEditChange /*= true*/) override
 	{
 		check (SkeletalMesh);
-		USkeleton* Skeleton = SkeletalMesh->Skeleton;
+		USkeleton* Skeleton = SkeletalMesh->GetSkeleton();
 		check (Skeleton);
 
 		// find all the bones to remove from Skeleton settings
@@ -370,7 +370,7 @@ public:
 
 		bool bNeedsRemoval = GetBoneReductionData(SkeletalMesh, DesiredLOD, BonesToRemove, BoneNamesToRemove);
 		// Always restore all previously removed bones if not contained by BonesToRemove
-		SkeletalMesh->CalculateRequiredBones(SkeletalMesh->GetImportedModel()->LODModels[DesiredLOD], SkeletalMesh->RefSkeleton, &BonesToRemove);
+		SkeletalMesh->CalculateRequiredBones(SkeletalMesh->GetImportedModel()->LODModels[DesiredLOD], SkeletalMesh->GetRefSkeleton(), &BonesToRemove);
 		
 		SkeletalMesh->ReleaseResources();
 		SkeletalMesh->ReleaseResourcesFence.Wait();
@@ -396,7 +396,7 @@ public:
 			{
 				for (const FBoneReference& BoneReference : SkeletalMesh->GetLODInfo(DesiredLOD)->BonesToRemove)
 				{
-					int32 BoneIndex = SkeletalMesh->RefSkeleton.FindRawBoneIndex(BoneReference.BoneName);
+					int32 BoneIndex = SkeletalMesh->GetRefSkeleton().FindRawBoneIndex(BoneReference.BoneName);
 					if (BoneIndex != INDEX_NONE)
 					{
 						BoneIndices.AddUnique(BoneIndex);
diff --git a/Engine/Source/Developer/MeshBuilder/Private/SkeletalMeshBuilder.cpp b/Engine/Source/Developer/MeshBuilder/Private/SkeletalMeshBuilder.cpp
--- a/Engine/Source/Developer/MeshBuilder/Private/SkeletalMeshBuilder.cpp
+++ b/Engine/Source/Developer/MeshBuilder/Private/SkeletalMeshBuilder.cpp
@@ -65,7 +65,7 @@ bool FSkeletalMeshBuilder::Build(USkeletalMesh* SkeletalMesh, const int32 LODInd
 	//We want to backup in case the LODModel is regenerated, this data is use to validate in the UI if the ddc must be rebuild
 	const FString BackupBuildStringID = SkeletalMesh->GetImportedModel()->LODModels[LODIndex].BuildStringID;
 
-	const FReferenceSkeleton& RefSkeleton = SkeletalMesh->RefSkeleton;
+	const FReferenceSkeleton& RefSkeleton = SkeletalMesh->GetRefSkeleton();
 
 	FScopedSlowTask SlowTask(6.01f, NSLOCTEXT("SkeltalMeshBuilder", "BuildingSkeletalMeshLOD", "Building skeletal mesh LOD"));
 	SlowTask.MakeDialog();
@@ -97,7 +97,7 @@ bool FSkeletalMeshBuilder::Build(USkeletalMesh* SkeletalMesh, const int32 LODInd
 		NumTextCoord = FMath::Max<int32>(NumTextCoord, SkeletalMeshImportData.NumTexCoords);
 		
 		//BaseLOD need to make sure the source data fit with the skeletalmesh materials array before using meshutilities.BuildSkeletalMesh
-		FLODUtilities::AdjustImportDataFaceMaterialIndex(SkeletalMesh->Materials, SkeletalMeshImportData.Materials, LODFaces, LODIndex);
+		FLODUtilities::AdjustImportDataFaceMaterialIndex(SkeletalMesh->GetMaterials(), SkeletalMeshImportData.Materials, LODFaces, LODIndex);
 
 		//Build the skeletalmesh using mesh utilities module
 		IMeshUtilities::MeshBuildOptions Options;
@@ -156,8 +156,8 @@ bool FSkeletalMeshBuilder::Build(USkeletalMesh* SkeletalMesh, const int32 LODInd
 				//Make the copy of the data only once until the ImportedModel change (re-imported)
 				SkeletalMesh->GetImportedModel()->OriginalReductionSourceMeshData[LODIndex]->EmptyBulkData();
 				TMap<FString, TArray<FMorphTargetDelta>> BaseLODMorphTargetData;
-				BaseLODMorphTargetData.Empty(SkeletalMesh->MorphTargets.Num());
-				for (UMorphTarget *MorphTarget : SkeletalMesh->MorphTargets)
+				BaseLODMorphTargetData.Empty(SkeletalMesh->GetMorphTargets().Num());
+				for (UMorphTarget *MorphTarget : SkeletalMesh->GetMorphTargets())
 				{
 					if (!MorphTarget->HasDataForLOD(LODIndex))
 					{
@@ -176,7 +176,7 @@ bool FSkeletalMeshBuilder::Build(USkeletalMesh* SkeletalMesh, const int32 LODInd
 
 				if (LODIndex == 0)
 				{
-					SkeletalMesh->GetLODInfo(LODIndex)->SourceImportFilename = SkeletalMesh->AssetImportData->GetFirstFilename();
+					SkeletalMesh->GetLODInfo(LODIndex)->SourceImportFilename = SkeletalMesh->GetAssetImportData()->GetFirstFilename();
 				}
 			}
 			FLODUtilities::SimplifySkeletalMeshLOD(UpdateContext, LODIndex, false);
diff --git a/Engine/Source/Developer/MeshMergeUtilities/Private/MeshMergeHelpers.cpp b/Engine/Source/Developer/MeshMergeUtilities/Private/MeshMergeHelpers.cpp
--- a/Engine/Source/Developer/MeshMergeUtilities/Private/MeshMergeHelpers.cpp
+++ b/Engine/Source/Developer/MeshMergeUtilities/Private/MeshMergeHelpers.cpp
@@ -328,10 +328,10 @@ void FMeshMergeHelpers::RetrieveMesh(const USkeletalMeshComponent* SkeletalMeshC
 			// use the remapping of material indices if there is a valid value
 			if (SrcLODInfo.LODMaterialMap.IsValidIndex(SectionIndex) && SrcLODInfo.LODMaterialMap[SectionIndex] != INDEX_NONE)
 			{
-				MaterialIndex = FMath::Clamp<int32>(SrcLODInfo.LODMaterialMap[SectionIndex], 0, SkeletalMeshComponent->SkeletalMesh->Materials.Num());
+				MaterialIndex = FMath::Clamp<int32>(SrcLODInfo.LODMaterialMap[SectionIndex], 0, SkeletalMeshComponent->SkeletalMesh->GetMaterials().Num());
 			}
 
-			FName ImportedMaterialSlotName = SkeletalMeshComponent->SkeletalMesh->Materials[MaterialIndex].ImportedMaterialSlotName;
+			FName ImportedMaterialSlotName = SkeletalMeshComponent->SkeletalMesh->GetMaterials()[MaterialIndex].ImportedMaterialSlotName;
 			const FPolygonGroupID SectionPolygonGroupID(SectionIndex);
 			if (!OutMeshDescription.IsPolygonGroupValid(SectionPolygonGroupID))
 			{
diff --git a/Engine/Source/Developer/MeshMergeUtilities/Private/SkeletalMeshAdapter.cpp b/Engine/Source/Developer/MeshMergeUtilities/Private/SkeletalMeshAdapter.cpp
--- a/Engine/Source/Developer/MeshMergeUtilities/Private/SkeletalMeshAdapter.cpp
+++ b/Engine/Source/Developer/MeshMergeUtilities/Private/SkeletalMeshAdapter.cpp
@@ -54,19 +54,19 @@ FString FSkeletalMeshComponentAdapter::GetBaseName() const
 
 FName FSkeletalMeshComponentAdapter::GetMaterialSlotName(int32 MaterialIndex) const
 {
-	return SkeletalMesh->Materials[MaterialIndex].MaterialSlotName;
+	return SkeletalMesh->GetMaterials()[MaterialIndex].MaterialSlotName;
 }
 
 FName FSkeletalMeshComponentAdapter::GetImportedMaterialSlotName(int32 MaterialIndex) const
 {
-	return SkeletalMesh->Materials[MaterialIndex].ImportedMaterialSlotName;
+	return SkeletalMesh->GetMaterials()[MaterialIndex].ImportedMaterialSlotName;
 }
 
 void FSkeletalMeshComponentAdapter::SetMaterial(int32 MaterialIndex, UMaterialInterface* Material)
 {
 	//We need to preserve the original material slot data
-	const FSkeletalMaterial& OriginalMaterialSlot = SkeletalMesh->Materials[MaterialIndex];
-	SkeletalMesh->Materials[MaterialIndex] = FSkeletalMaterial(Material, true, false, OriginalMaterialSlot.MaterialSlotName, OriginalMaterialSlot.ImportedMaterialSlotName);
+	const FSkeletalMaterial& OriginalMaterialSlot = SkeletalMesh->GetMaterials()[MaterialIndex];
+	SkeletalMesh->GetMaterials()[MaterialIndex] = FSkeletalMaterial(Material, true, false, OriginalMaterialSlot.MaterialSlotName, OriginalMaterialSlot.ImportedMaterialSlotName);
 }
 
 void FSkeletalMeshComponentAdapter::RemapMaterialIndex(int32 LODIndex, int32 SectionIndex, int32 NewMaterialIndex)
@@ -90,7 +90,7 @@ void FSkeletalMeshComponentAdapter::RemapMaterialIndex(int32 LODIndex, int32 Sec
 
 int32 FSkeletalMeshComponentAdapter::AddMaterial(UMaterialInterface* Material)
 {
-	const int32 Index = SkeletalMesh->Materials.Emplace(Material);
+	const int32 Index = SkeletalMesh->GetMaterials().Emplace(Material);
 	IMeshUtilities& MeshUtilities = FModuleManager::Get().LoadModuleChecked<IMeshUtilities>("MeshUtilities");
 	MeshUtilities.FixupMaterialSlotNames(SkeletalMesh);
 	return Index;
@@ -98,7 +98,7 @@ int32 FSkeletalMeshComponentAdapter::AddMaterial(UMaterialInterface* Material)
 
 int32 FSkeletalMeshComponentAdapter::AddMaterial(UMaterialInterface* Material, const FName& SlotName, const FName& ImportedSlotName)
 {
-	const int32 Index = SkeletalMesh->Materials.Emplace(Material, true, false, SlotName, ImportedSlotName);
+	const int32 Index = SkeletalMesh->GetMaterials().Emplace(Material, true, false, SlotName, ImportedSlotName);
 	IMeshUtilities& MeshUtilities = FModuleManager::Get().LoadModuleChecked<IMeshUtilities>("MeshUtilities");
 	MeshUtilities.FixupMaterialSlotNames(SkeletalMesh);
 	return Index;
diff --git a/Engine/Source/Developer/MeshUtilities/Private/MeshUtilities.cpp b/Engine/Source/Developer/MeshUtilities/Private/MeshUtilities.cpp
--- a/Engine/Source/Developer/MeshUtilities/Private/MeshUtilities.cpp
+++ b/Engine/Source/Developer/MeshUtilities/Private/MeshUtilities.cpp
@@ -306,7 +306,7 @@ static void SkinnedMeshToRawMeshes(USkinnedMeshComponent* InSkinnedMeshComponent
 				// use the remapping of material indices if there is a valid value
 				if (SrcLODInfo.LODMaterialMap.IsValidIndex(SectionIndex) && SrcLODInfo.LODMaterialMap[SectionIndex] != INDEX_NONE)
 				{
-					MaterialIndex = FMath::Clamp<int32>(SrcLODInfo.LODMaterialMap[SectionIndex], 0, InSkinnedMeshComponent->SkeletalMesh->Materials.Num());
+					MaterialIndex = FMath::Clamp<int32>(SrcLODInfo.LODMaterialMap[SectionIndex], 0, InSkinnedMeshComponent->SkeletalMesh->GetMaterials().Num());
 				}
 
 				// copy face info
@@ -3055,7 +3055,7 @@ void FMeshUtilities::FixupMaterialSlotNames(UStaticMesh* StaticMesh) const
 
 void FMeshUtilities::FixupMaterialSlotNames(USkeletalMesh* SkeletalMesh) const
 {
-	FixupMaterialSlotNames_Implementation(SkeletalMesh->Materials);
+	FixupMaterialSlotNames_Implementation(SkeletalMesh->GetMaterials());
 }
 
 bool FMeshUtilities::BuildStaticMesh(FStaticMeshRenderData& OutRenderData, UStaticMesh* StaticMesh, const FStaticMeshLODGroup& LODGroup)
diff --git a/Engine/Source/Developer/MeshUtilities/Private/SkeletalMeshTools.cpp b/Engine/Source/Developer/MeshUtilities/Private/SkeletalMeshTools.cpp
--- a/Engine/Source/Developer/MeshUtilities/Private/SkeletalMeshTools.cpp
+++ b/Engine/Source/Developer/MeshUtilities/Private/SkeletalMeshTools.cpp
@@ -631,10 +631,10 @@ namespace SkeletalMeshTools
 			return;
 
 		SkeletalMesh->CalculateInvRefMatrices();
-		check( SkeletalMesh->RefSkeleton.GetRawBoneNum() == SkeletalMesh->RefBasesInvMatrix.Num() );
+		check( SkeletalMesh->GetRefSkeleton().GetRawBoneNum() == SkeletalMesh->GetRefBasesInvMatrix().Num() );
 
 		Infos.Empty();
-		Infos.AddZeroed( SkeletalMesh->RefSkeleton.GetRawBoneNum() );
+		Infos.AddZeroed( SkeletalMesh->GetRefSkeleton().GetRawBoneNum() );
 
 		FSkeletalMeshLODModel* LODModel = &ImportedResource->LODModels[0];
 		for(int32 SectionIndex = 0; SectionIndex < LODModel->Sections.Num(); SectionIndex++)
@@ -648,10 +648,10 @@ namespace SkeletalMeshTools
 				{
 					int32 BoneIndex = Section.BoneMap[GetDominantBoneIndex(SoftVert)];
 
-					FVector LocalPos = SkeletalMesh->RefBasesInvMatrix[BoneIndex].TransformPosition(SoftVert->Position);
+					FVector LocalPos = SkeletalMesh->GetRefBasesInvMatrix()[BoneIndex].TransformPosition(SoftVert->Position);
 					Infos[BoneIndex].Positions.Add(LocalPos);
 
-					FVector LocalNormal = SkeletalMesh->RefBasesInvMatrix[BoneIndex].TransformVector(SoftVert->TangentZ);
+					FVector LocalNormal = SkeletalMesh->GetRefBasesInvMatrix()[BoneIndex].TransformVector(SoftVert->TangentZ);
 					Infos[BoneIndex].Normals.Add(LocalNormal);
 				}
 				else
@@ -662,10 +662,10 @@ namespace SkeletalMeshTools
 						{
 							int32 BoneIndex = Section.BoneMap[SoftVert->InfluenceBones[j]];
 
-							FVector LocalPos = SkeletalMesh->RefBasesInvMatrix[BoneIndex].TransformPosition(SoftVert->Position);
+							FVector LocalPos = SkeletalMesh->GetRefBasesInvMatrix()[BoneIndex].TransformPosition(SoftVert->Position);
 							Infos[BoneIndex].Positions.Add(LocalPos);
 
-							FVector LocalNormal = SkeletalMesh->RefBasesInvMatrix[BoneIndex].TransformVector(SoftVert->TangentZ);
+							FVector LocalNormal = SkeletalMesh->GetRefBasesInvMatrix()[BoneIndex].TransformVector(SoftVert->TangentZ);
 							Infos[BoneIndex].Normals.Add(LocalNormal);
 						}
 					}
diff --git a/Engine/Source/Developer/SkeletalMeshUtilitiesCommon/Private/LODUtilities.cpp b/Engine/Source/Developer/SkeletalMeshUtilitiesCommon/Private/LODUtilities.cpp
--- a/Engine/Source/Developer/SkeletalMeshUtilitiesCommon/Private/LODUtilities.cpp
+++ b/Engine/Source/Developer/SkeletalMeshUtilitiesCommon/Private/LODUtilities.cpp
@@ -410,7 +410,7 @@ void FLODUtilities::RemoveLOD(FSkeletalMeshUpdateContext& UpdateContext, int32 D
 		}
 
 		//remove all Morph target data for this LOD
-		for (UMorphTarget* MorphTarget : SkeletalMesh->MorphTargets)
+		for (UMorphTarget* MorphTarget : SkeletalMesh->GetMorphTargets())
 		{
 			if (MorphTarget->HasDataForLOD(DesiredLOD))
 			{
@@ -850,7 +850,7 @@ void CreateLODMorphTarget(USkeletalMesh* SkeletalMesh, FReductionBaseSkeletalMes
 
 	bool bInitializeMorphData = false;
 
-	for (UMorphTarget *MorphTarget : SkeletalMesh->MorphTargets)
+	for (UMorphTarget *MorphTarget : SkeletalMesh->GetMorphTargets())
 	{
 		if (!MorphTarget->HasDataForLOD(SourceLOD))
 		{
@@ -972,7 +972,7 @@ void FLODUtilities::ClearGeneratedMorphTarget(USkeletalMesh* SkeletalMesh, int32
 
 	const FSkeletalMeshLODModel& TargetLODModel = SkeletalMeshResource->LODModels[TargetLOD];
 	//Make sure we have some morph for this LOD
-	for (UMorphTarget *MorphTarget : SkeletalMesh->MorphTargets)
+	for (UMorphTarget *MorphTarget : SkeletalMesh->GetMorphTargets())
 	{
 		if (!MorphTarget->HasDataForLOD(TargetLOD))
 		{
@@ -1016,7 +1016,7 @@ void FLODUtilities::ApplyMorphTargetsToLOD(USkeletalMesh* SkeletalMesh, int32 So
 
 	//Make sure we have some morph for this LOD
 	bool bContainsMorphTargets = false;
-	for (UMorphTarget* MorphTarget : SkeletalMesh->MorphTargets)
+	for (UMorphTarget* MorphTarget : SkeletalMesh->GetMorphTargets())
 	{
 		if (MorphTarget->HasDataForLOD(SourceLOD))
 		{
@@ -1130,7 +1130,7 @@ void FLODUtilities::ApplyMorphTargetsToLOD(USkeletalMesh* SkeletalMesh, int32 So
 	TMap<UMorphTarget *, TMap<uint32, uint32>> PerMorphTargetBaseIndexToMorphTargetDelta;
 	//Create a map from BaseIndex to a list of match target index for all base morph target point
 	TMap<uint32, TArray<uint32>> BaseMorphIndexToTargetIndexList;
-	for (UMorphTarget *MorphTarget : SkeletalMesh->MorphTargets)
+	for (UMorphTarget *MorphTarget : SkeletalMesh->GetMorphTargets())
 	{
 		if (!MorphTarget->HasDataForLOD(SourceLOD))
 		{
@@ -1241,8 +1241,8 @@ void FLODUtilities::SimplifySkeletalMeshLOD( USkeletalMesh* SkeletalMesh, int32
 			if (SkeletalMeshResource->OriginalReductionSourceMeshData[DesiredLOD]->IsEmpty())
 			{
 				TMap<FString, TArray<FMorphTargetDelta>> BaseLODMorphTargetData;
-				BaseLODMorphTargetData.Empty(SkeletalMesh->MorphTargets.Num());
-				for (UMorphTarget *MorphTarget : SkeletalMesh->MorphTargets)
+				BaseLODMorphTargetData.Empty(SkeletalMesh->GetMorphTargets().Num());
+				for (UMorphTarget *MorphTarget : SkeletalMesh->GetMorphTargets())
 				{
 					if (!MorphTarget->HasDataForLOD(DesiredLOD))
 					{
@@ -1273,7 +1273,7 @@ void FLODUtilities::SimplifySkeletalMeshLOD( USkeletalMesh* SkeletalMesh, int32
 
 				if (DesiredLOD == 0)
 				{
-					SkeletalMesh->GetLODInfo(DesiredLOD)->SourceImportFilename = SkeletalMesh->AssetImportData->GetFirstFilename();
+					SkeletalMesh->GetLODInfo(DesiredLOD)->SourceImportFilename = SkeletalMesh->GetAssetImportData()->GetFirstFilename();
 				}
 			}
 		}
@@ -1377,7 +1377,7 @@ bool FLODUtilities::RestoreSkeletalMeshLODImportedData(USkeletalMesh* SkeletalMe
 		FSkeletalMeshLODModel::CopyStructure(&(SkeletalMesh->GetImportedModel()->LODModels[LodIndex]), &ImportedBaseLODModel);
 		//Copy the morph target deltas
 		bool bInitMorphTargetData = false;
-		for (UMorphTarget *MorphTarget : SkeletalMesh->MorphTargets)
+		for (UMorphTarget *MorphTarget : SkeletalMesh->GetMorphTargets())
 		{
 			if (!ImportedBaseLODMorphTargetData.Contains(MorphTarget->GetFullName()))
 			{
@@ -1804,7 +1804,7 @@ bool FLODUtilities::UpdateAlternateSkinWeights(USkeletalMesh* SkeletalMeshDest,
 	}
 	FSkeletalMeshImportData ImportDataDest;
 	SkeletalMeshDest->LoadLODImportedData(LODIndexDest, ImportDataDest);
-	return UpdateAlternateSkinWeights(LODModelDest, ImportDataDest, SkeletalMeshDest, SkeletalMeshDest->RefSkeleton, ProfileNameDest, LODIndexDest, OverlappingThresholds, ShouldImportNormals, ShouldImportTangents, bUseMikkTSpace, bComputeWeightedNormals);
+	return UpdateAlternateSkinWeights(LODModelDest, ImportDataDest, SkeletalMeshDest, SkeletalMeshDest->GetRefSkeleton(), ProfileNameDest, LODIndexDest, OverlappingThresholds, ShouldImportNormals, ShouldImportTangents, bUseMikkTSpace, bComputeWeightedNormals);
 }
 
 bool FLODUtilities::UpdateAlternateSkinWeights(FSkeletalMeshLODModel& LODModelDest, FSkeletalMeshImportData& ImportDataDest, USkeletalMesh* SkeletalMeshDest, FReferenceSkeleton& RefSkeleton, const FName& ProfileNameDest, int32 LODIndexDest, FOverlappingThresholds OverlappingThresholds, bool ShouldImportNormals, bool ShouldImportTangents, bool bUseMikkTSpace, bool bComputeWeightedNormals)
@@ -2123,7 +2123,7 @@ bool FLODUtilities::UpdateAlternateSkinWeights(FSkeletalMeshLODModel& LODModelDe
 	TArray<FName> WarningNames;
 
 	//BaseLOD need to make sure the source data fit with the skeletalmesh materials array before using meshutilities.BuildSkeletalMesh
-	AdjustImportDataFaceMaterialIndex(SkeletalMeshDest->Materials, ImportDataDest.Materials, LODFacesDest, LODIndexDest);
+	AdjustImportDataFaceMaterialIndex(SkeletalMeshDest->GetMaterials(), ImportDataDest.Materials, LODFacesDest, LODIndexDest);
 
 	//Build the destination mesh with the Alternate influences, so the chunking is done properly.
 	bBuildSuccess = MeshUtilities.BuildSkeletalMesh(LODModelDest, SkeletalMeshDest->GetName(), RefSkeleton, LODInfluencesDest, LODWedgesDest, LODFacesDest, LODPointsDest, LODPointToRawMapDest, BuildOptions, &WarningMessages, &WarningNames);
@@ -2824,7 +2824,7 @@ void FLODUtilities::BuildMorphTargets(USkeletalMesh* BaseSkelMesh, FSkeletalMesh
 
 			TArray< FMorphTargetDelta >* Deltas = Results[NewMorphDeltasIdx];
 
-			FAsyncTask<FAsyncImportMorphTargetWork>* NewWork = new FAsyncTask<FAsyncImportMorphTargetWork>(&BaseLODModel, BaseSkelMesh->RefSkeleton, BaseImportData,
+			FAsyncTask<FAsyncImportMorphTargetWork>* NewWork = new FAsyncTask<FAsyncImportMorphTargetWork>(&BaseLODModel, BaseSkelMesh->GetRefSkeleton(), BaseImportData,
 				MoveTemp(ShapeImportData.Points), *Deltas, BaseIndexData, BaseWedgePointIndices, WedgePointToVertexIndexMap, OverlappingVertices, MoveTemp(ModifiedPoints), WedgeToFaces, MeshDataBundle, TangentZ,
 				ShouldImportNormals, ShouldImportTangents, bUseMikkTSpace, Thresholds);
 			PendingWork.Add(NewWork);
diff --git a/Engine/Source/Editor/AnimGraph/Private/AnimGraphNode_SkeletalControlBase.cpp b/Engine/Source/Editor/AnimGraph/Private/AnimGraphNode_SkeletalControlBase.cpp
--- a/Engine/Source/Editor/AnimGraph/Private/AnimGraphNode_SkeletalControlBase.cpp
+++ b/Engine/Source/Editor/AnimGraph/Private/AnimGraphNode_SkeletalControlBase.cpp
@@ -81,7 +81,7 @@ void UAnimGraphNode_SkeletalControlBase::CreateOutputPins()
 
 void UAnimGraphNode_SkeletalControlBase::ConvertToComponentSpaceTransform(const USkeletalMeshComponent* SkelComp, const FTransform & InTransform, FTransform & OutCSTransform, int32 BoneIndex, EBoneControlSpace Space) const
 {
-	USkeleton * Skeleton = SkelComp->SkeletalMesh->Skeleton;
+	USkeleton * Skeleton = SkelComp->SkeletalMesh->GetSkeleton();
 
 	switch (Space)
 	{
@@ -250,7 +250,7 @@ FVector UAnimGraphNode_SkeletalControlBase::ConvertWidgetLocation(const USkeleta
 
 	if (MeshBases.GetPose().IsValid())
 	{
-		USkeleton * Skeleton = SkelComp->SkeletalMesh->Skeleton;
+		USkeleton * Skeleton = SkelComp->SkeletalMesh->GetSkeleton();
 		const FMeshPoseBoneIndex MeshBoneIndex(SkelComp->GetBoneIndex(BoneName));
 		const FCompactPoseBoneIndex CompactBoneIndex = MeshBases.GetPose().GetBoneContainer().MakeCompactPoseIndex(MeshBoneIndex);
 
diff --git a/Engine/Source/Editor/AnimGraph/Private/AnimGraphNode_Trail.cpp b/Engine/Source/Editor/AnimGraph/Private/AnimGraphNode_Trail.cpp
--- a/Engine/Source/Editor/AnimGraph/Private/AnimGraphNode_Trail.cpp
+++ b/Engine/Source/Editor/AnimGraph/Private/AnimGraphNode_Trail.cpp
@@ -137,7 +137,7 @@ void UAnimGraphNode_Trail::Draw(FPrimitiveDrawInterface* PDI, USkeletalMeshCompo
 		int32 CurrentIndex = Node.ChainLength - 1;
 		TrailBoneList[CurrentIndex] = Node.TrailBone.BoneName;
 		FName CurrentName = Node.TrailBone.BoneName;
-		const FReferenceSkeleton& RefSkeleton = PreviewSkelMeshComp->SkeletalMesh->RefSkeleton;
+		const FReferenceSkeleton& RefSkeleton = PreviewSkelMeshComp->SkeletalMesh->GetRefSkeleton();
 		while (--CurrentIndex >= 0 && CurrentName != NAME_None)
 		{
 			const int32 ParentIndex = RefSkeleton.GetParentIndex(RefSkeleton.FindBoneIndex(CurrentName));
diff --git a/Engine/Source/Editor/AnimGraph/Private/EditModes/AnimNodeEditMode.cpp b/Engine/Source/Editor/AnimGraph/Private/EditModes/AnimNodeEditMode.cpp
--- a/Engine/Source/Editor/AnimGraph/Private/EditModes/AnimNodeEditMode.cpp
+++ b/Engine/Source/Editor/AnimGraph/Private/EditModes/AnimNodeEditMode.cpp
@@ -357,7 +357,7 @@ void FAnimNodeEditMode::Tick(FEditorViewportClient* ViewportClient, float DeltaT
 
 void FAnimNodeEditMode::ConvertToComponentSpaceTransform(const USkeletalMeshComponent* SkelComp, const FTransform & InTransform, FTransform & OutCSTransform, int32 BoneIndex, EBoneControlSpace Space)
 {
-	USkeleton* Skeleton = SkelComp->SkeletalMesh->Skeleton;
+	USkeleton* Skeleton = SkelComp->SkeletalMesh->GetSkeleton();
 
 	switch (Space)
 	{
@@ -427,7 +427,7 @@ void FAnimNodeEditMode::ConvertToComponentSpaceTransform(const USkeletalMeshComp
 
 void FAnimNodeEditMode::ConvertToBoneSpaceTransform(const USkeletalMeshComponent* SkelComp, const FTransform & InCSTransform, FTransform & OutBSTransform, int32 BoneIndex, EBoneControlSpace Space)
 {
-	USkeleton* Skeleton = SkelComp->SkeletalMesh->Skeleton;
+	USkeleton* Skeleton = SkelComp->SkeletalMesh->GetSkeleton();
 
 	switch(Space)
 	{
@@ -691,7 +691,7 @@ FVector FAnimNodeEditMode::ConvertWidgetLocation(const USkeletalMeshComponent* S
 	{
 		if (InMeshBases.GetPose().IsValid())
 		{
-			USkeleton* Skeleton = InSkelComp->SkeletalMesh->Skeleton;
+			USkeleton* Skeleton = InSkelComp->SkeletalMesh->GetSkeleton();
 			const int32 MeshBoneIndex = InSkelComp->GetBoneIndex(InBoneName);
 			if (MeshBoneIndex != INDEX_NONE)
 			{
diff --git a/Engine/Source/Editor/AnimGraph/Private/EditModes/ModifyBoneEditMode.cpp b/Engine/Source/Editor/AnimGraph/Private/EditModes/ModifyBoneEditMode.cpp
--- a/Engine/Source/Editor/AnimGraph/Private/EditModes/ModifyBoneEditMode.cpp
+++ b/Engine/Source/Editor/AnimGraph/Private/EditModes/ModifyBoneEditMode.cpp
@@ -59,7 +59,7 @@ ECoordSystem FModifyBoneEditMode::GetWidgetCoordinateSystem() const
 FVector FModifyBoneEditMode::GetWidgetLocation() const
 {
 	USkeletalMeshComponent* SkelComp = GetAnimPreviewScene().GetPreviewMeshComponent();
-	USkeleton* Skeleton = SkelComp->SkeletalMesh->Skeleton;
+	USkeleton* Skeleton = SkelComp->SkeletalMesh->GetSkeleton();
 	FVector WidgetLoc = FVector::ZeroVector;
 
 	// if the current widget mode is translate, then shows the widget according to translation space
diff --git a/Engine/Source/Editor/AnimGraph/Private/EditModes/ObserveBoneEditMode.cpp b/Engine/Source/Editor/AnimGraph/Private/EditModes/ObserveBoneEditMode.cpp
--- a/Engine/Source/Editor/AnimGraph/Private/EditModes/ObserveBoneEditMode.cpp
+++ b/Engine/Source/Editor/AnimGraph/Private/EditModes/ObserveBoneEditMode.cpp
@@ -40,7 +40,7 @@ ECoordSystem FObserveBoneEditMode::GetWidgetCoordinateSystem() const
 FVector FObserveBoneEditMode::GetWidgetLocation() const
 {
 	USkeletalMeshComponent* SkelComp = GetAnimPreviewScene().GetPreviewMeshComponent();
-	USkeleton* Skeleton = SkelComp->SkeletalMesh->Skeleton;
+	USkeleton* Skeleton = SkelComp->SkeletalMesh->GetSkeleton();
 	FVector WidgetLoc = FVector::ZeroVector;
 
 	int32 MeshBoneIndex = SkelComp->GetBoneIndex(GraphNode->Node.BoneToObserve.BoneName);
diff --git a/Engine/Source/Editor/AnimGraph/Private/EditModes/TwoBoneIKEditMode.cpp b/Engine/Source/Editor/AnimGraph/Private/EditModes/TwoBoneIKEditMode.cpp
--- a/Engine/Source/Editor/AnimGraph/Private/EditModes/TwoBoneIKEditMode.cpp
+++ b/Engine/Source/Editor/AnimGraph/Private/EditModes/TwoBoneIKEditMode.cpp
@@ -57,9 +57,9 @@ void FTwoBoneIKEditMode::ExitMode()
 void FTwoBoneIKEditMode::Render(const FSceneView* View, FViewport* Viewport, FPrimitiveDrawInterface* PDI)
 {
 	USkeletalMeshComponent* SkelMeshComp = GetAnimPreviewScene().GetPreviewMeshComponent();
-	if (SkelMeshComp && SkelMeshComp->SkeletalMesh && SkelMeshComp->SkeletalMesh->Skeleton)
+	if (SkelMeshComp && SkelMeshComp->SkeletalMesh && SkelMeshComp->SkeletalMesh->GetSkeleton())
 	{
-		USkeleton* Skeleton = SkelMeshComp->SkeletalMesh->Skeleton;
+		USkeleton* Skeleton = SkelMeshComp->SkeletalMesh->GetSkeleton();
 
 		PDI->SetHitProxy(new HTwoBoneIKProxy(BSM_EndEffector));
 		DrawTargetLocation(PDI, BSM_EndEffector, FColor(255, 128, 128), FColor(180, 128, 128));
diff --git a/Engine/Source/Editor/AnimationModifiers/Private/AnimationBlueprintLibrary.cpp b/Engine/Source/Editor/AnimationModifiers/Private/AnimationBlueprintLibrary.cpp
--- a/Engine/Source/Editor/AnimationModifiers/Private/AnimationBlueprintLibrary.cpp
+++ b/Engine/Source/Editor/AnimationModifiers/Private/AnimationBlueprintLibrary.cpp
@@ -1809,7 +1809,7 @@ void UAnimationBlueprintLibrary::GetBonePosesForTime(const UAnimSequence* Animat
 		// Need this for FCompactPose
 		FMemMark Mark(FMemStack::Get());
 
-		const FReferenceSkeleton& RefSkeleton = (PreviewMesh)? PreviewMesh->RefSkeleton : AnimationSequence->GetSkeleton()->GetReferenceSkeleton();
+		const FReferenceSkeleton& RefSkeleton = (PreviewMesh)? PreviewMesh->GetRefSkeleton() : AnimationSequence->GetSkeleton()->GetReferenceSkeleton();
 
 		if (IsValidTimeInternal(AnimationSequence, Time))
 		{
diff --git a/Engine/Source/Editor/ClothPainter/Private/ClothMeshAdapter.cpp b/Engine/Source/Editor/ClothPainter/Private/ClothMeshAdapter.cpp
--- a/Engine/Source/Editor/ClothPainter/Private/ClothMeshAdapter.cpp
+++ b/Engine/Source/Editor/ClothPainter/Private/ClothMeshAdapter.cpp
@@ -193,7 +193,7 @@ void FClothMeshPaintAdapter::SetSelectedClothingAsset(const FGuid& InAssetGuid,
 	SelectedAsset = nullptr;
 	if(InAssetGuid.IsValid() && ReferencedSkeletalMesh)
 	{
-		for(UClothingAssetBase* Asset : ReferencedSkeletalMesh->MeshClothingAssets)
+		for(UClothingAssetBase* Asset : ReferencedSkeletalMesh->GetMeshClothingAssets())
 		{
 			UClothingAssetCommon* ConcreteAsset = CastChecked<UClothingAssetCommon>(Asset);
 			if(ConcreteAsset->GetAssetGuid() == InAssetGuid)
diff --git a/Engine/Source/Editor/ClothPainter/Private/ClothPainter.cpp b/Engine/Source/Editor/ClothPainter/Private/ClothPainter.cpp
--- a/Engine/Source/Editor/ClothPainter/Private/ClothPainter.cpp
+++ b/Engine/Source/Editor/ClothPainter/Private/ClothPainter.cpp
@@ -177,7 +177,7 @@ void FClothPainter::RefreshClothingAssets()
 
 	if(USkeletalMesh* Mesh = SkeletalMeshComponent->SkeletalMesh)
 	{
-		for(UClothingAssetBase* BaseClothingAsset : Mesh->MeshClothingAssets)
+		for(UClothingAssetBase* BaseClothingAsset : Mesh->GetMeshClothingAssets())
 		{
 			if(UClothingAssetCommon* ActualAsset = Cast<UClothingAssetCommon>(BaseClothingAsset))
 			{
diff --git a/Engine/Source/Editor/ClothPainter/Private/ClothingPaintEditMode.cpp b/Engine/Source/Editor/ClothPainter/Private/ClothingPaintEditMode.cpp
--- a/Engine/Source/Editor/ClothPainter/Private/ClothingPaintEditMode.cpp
+++ b/Engine/Source/Editor/ClothPainter/Private/ClothingPaintEditMode.cpp
@@ -90,7 +90,7 @@ void FClothingPaintEditMode::Exit()
 
 			if(USkeletalMesh* SkelMesh = MeshComponent->SkeletalMesh)
 			{
-				for(UClothingAssetBase* AssetBase : SkelMesh->MeshClothingAssets)
+				for(UClothingAssetBase* AssetBase : SkelMesh->GetMeshClothingAssets())
 				{
 					UClothingAssetCommon* ConcreteAsset = CastChecked<UClothingAssetCommon>(AssetBase);
 					ConcreteAsset->ApplyParameterMasks();
diff --git a/Engine/Source/Editor/ClothPainter/Private/SClothAssetSelector.cpp b/Engine/Source/Editor/ClothPainter/Private/SClothAssetSelector.cpp
--- a/Engine/Source/Editor/ClothPainter/Private/SClothAssetSelector.cpp
+++ b/Engine/Source/Editor/ClothPainter/Private/SClothAssetSelector.cpp
@@ -210,14 +210,14 @@ private:
 			{
 				FScopedSuspendAlternateSkinWeightPreview ScopedSuspendAlternateSkinnWeightPreview(SkelMesh);
 				int32 AssetIndex;
-				if(SkelMesh->MeshClothingAssets.Find(Asset, AssetIndex))
+				if(SkelMesh->GetMeshClothingAssets().Find(Asset, AssetIndex))
 				{
 					// Need to unregister our components so they shut down their current clothing simulation
 					FScopedSkeletalMeshPostEditChange ScopedPostEditChange(SkelMesh);
 					SkelMesh->PreEditChange(nullptr);
 
 					Asset->UnbindFromSkeletalMesh(SkelMesh);
-					SkelMesh->MeshClothingAssets.RemoveAt(AssetIndex);
+					SkelMesh->GetMeshClothingAssets().RemoveAt(AssetIndex);
 
 					// Need to fix up asset indices on sections.
 					if(FSkeletalMeshModel* Model = SkelMesh->GetImportedModel())
@@ -1179,7 +1179,7 @@ void SClothAssetSelector::OnCopyClothingAssetSelected(const FAssetData& AssetDat
 		FClothingSystemEditorInterfaceModule& ClothingEditorModule = FModuleManager::LoadModuleChecked<FClothingSystemEditorInterfaceModule>("ClothingSystemEditorInterface");
 		UClothingAssetFactoryBase* AssetFactory = ClothingEditorModule.GetClothingAssetFactory();
 
-		for (UClothingAssetBase* ClothingAsset : SourceSkelMesh->MeshClothingAssets)
+		for (UClothingAssetBase* ClothingAsset : SourceSkelMesh->GetMeshClothingAssets())
 		{
 			UClothingAssetCommon* NewAsset = Cast<UClothingAssetCommon>(AssetFactory->CreateFromExistingCloth(Mesh, SourceSkelMesh, ClothingAsset));
 			Mesh->AddClothingAsset(NewAsset);
@@ -1369,7 +1369,7 @@ void SClothAssetSelector::RefreshAssetList()
 
 	AssetListItems.Empty();
 
-	for(UClothingAssetBase* Asset : Mesh->MeshClothingAssets)
+	for(UClothingAssetBase* Asset : Mesh->GetMeshClothingAssets())
 	{
 		UClothingAssetCommon* ConcreteAsset = Cast<UClothingAssetCommon>(Asset);
 
diff --git a/Engine/Source/Editor/ClothingSystemEditor/Private/ClothingAssetFactory.cpp b/Engine/Source/Editor/ClothingSystemEditor/Private/ClothingAssetFactory.cpp
--- a/Engine/Source/Editor/ClothingSystemEditor/Private/ClothingAssetFactory.cpp
+++ b/Engine/Source/Editor/ClothingSystemEditor/Private/ClothingAssetFactory.cpp
@@ -137,7 +137,7 @@ UClothingAssetBase* UClothingAssetFactory::Reimport(const FString& Filename, USk
 	}
 
 	int32 OldIndex = INDEX_NONE;
-	if(TargetMesh->MeshClothingAssets.Find(OriginalAsset, OldIndex))
+	if(TargetMesh->GetMeshClothingAssets().Find(OriginalAsset, OldIndex))
 	{
 		TArray<UActorComponent*> ComponentsToReregister;
 		for(TObjectIterator<USkeletalMeshComponent> It; It; ++It)
@@ -153,11 +153,11 @@ UClothingAssetBase* UClothingAssetFactory::Reimport(const FString& Filename, USk
 	
 		FMultiComponentReregisterContext ReregisterContext(ComponentsToReregister);
 
-		UClothingAssetCommon* OldClothingAsset = Cast<UClothingAssetCommon>(TargetMesh->MeshClothingAssets[OldIndex]);
+		UClothingAssetCommon* OldClothingAsset = Cast<UClothingAssetCommon>(TargetMesh->GetMeshClothingAssets()[OldIndex]);
 		UClothingAssetCommon* NewClothingAsset = nullptr;
 		FName AssetName = NAME_None;
 
-		if(!OldClothingAsset || !TargetMesh->MeshClothingAssets.IsValidIndex(OldIndex))
+		if(!OldClothingAsset || !TargetMesh->GetMeshClothingAssets().IsValidIndex(OldIndex))
 		{
 			return nullptr;
 		}
@@ -182,7 +182,7 @@ UClothingAssetBase* UClothingAssetFactory::Reimport(const FString& Filename, USk
 				// Store import path
 				NewClothingAsset->ImportedFilePath = Filename;
 
-				TargetMesh->MeshClothingAssets[OldIndex] = NewClothingAsset;
+				TargetMesh->GetMeshClothingAssets()[OldIndex] = NewClothingAsset;
 
 				for(ClothingAssetUtils::FClothingAssetMeshBinding& Binding : AssetBindings)
 				{
@@ -320,7 +320,7 @@ UClothingAssetBase* UClothingAssetFactory::ImportLodToClothing(USkeletalMesh* Ta
 	{
 		// Find the clothing asset in the mesh to verify the params are correct
 		int32 MeshAssetIndex = INDEX_NONE;
-		if(TargetMesh->MeshClothingAssets.Find(TargetClothing, MeshAssetIndex))
+		if(TargetMesh->GetMeshClothingAssets().Find(TargetClothing, MeshAssetIndex))
 		{
 			// Everything looks good, continue to actual import
 			UClothingAssetCommon* ConcreteTarget = CastChecked<UClothingAssetCommon>(TargetClothing);
@@ -444,7 +444,7 @@ UClothingAssetBase* UClothingAssetFactory::CreateFromApexAsset(nvidia::apex::Clo
 	for(int32 UsedBoneIndex = 0; UsedBoneIndex < NumUsedBones; ++UsedBoneIndex)
 	{
 		const FName& BoneName = NewClothingAsset->UsedBoneNames[UsedBoneIndex];
-		const int32 UnrealBoneIndex = TargetMesh->RefSkeleton.FindBoneIndex(BoneName);
+		const int32 UnrealBoneIndex = TargetMesh->GetRefSkeleton().FindBoneIndex(BoneName);
 
 		// If we find an invalid bone then the asset is invalid, as it cannot be skinned to this mesh
 		if(UnrealBoneIndex == INDEX_NONE)
@@ -1245,7 +1245,7 @@ bool UClothingAssetFactory::ImportToLodInternal(
 			const uint16 SourceIndex = SourceSection.BoneMap[SourceVert.InfluenceBones[InfluenceIndex]];
 			if(SourceIndex != INDEX_NONE)
 			{
-				FName BoneName = SourceMesh->RefSkeleton.GetBoneName(SourceIndex);
+				FName BoneName = SourceMesh->GetRefSkeleton().GetBoneName(SourceIndex);
 				BoneData.BoneIndices[InfluenceIndex] = DestAsset->UsedBoneNames.AddUnique(BoneName);
 				BoneData.BoneWeights[InfluenceIndex] = (float)SourceVert.InfluenceWeights[InfluenceIndex] / 255.0f;
 			}
diff --git a/Engine/Source/Editor/DetailCustomizations/Private/SkeletalMeshComponentDetails.cpp b/Engine/Source/Editor/DetailCustomizations/Private/SkeletalMeshComponentDetails.cpp
--- a/Engine/Source/Editor/DetailCustomizations/Private/SkeletalMeshComponentDetails.cpp
+++ b/Engine/Source/Editor/DetailCustomizations/Private/SkeletalMeshComponentDetails.cpp
@@ -398,14 +398,14 @@ USkeleton* FSkeletalMeshComponentDetails::GetValidSkeletonFromRegisteredMeshes()
 		// If we've not come across a valid skeleton yet, store this one.
 		if (!Skeleton)
 		{
-			Skeleton = Mesh->SkeletalMesh->Skeleton;
+			Skeleton = Mesh->SkeletalMesh->GetSkeleton();
 			continue;
 		}
 
 		// We've encountered a valid skeleton before.
 		// If this skeleton is not the same one, that means there are multiple
 		// skeletons selected, so we don't want to take any action.
-		if (Mesh->SkeletalMesh->Skeleton != Skeleton)
+		if (Mesh->SkeletalMesh->GetSkeleton() != Skeleton)
 		{
 			return NULL;
 		}
diff --git a/Engine/Source/Editor/LevelEditor/Private/LevelEditorActions.cpp b/Engine/Source/Editor/LevelEditor/Private/LevelEditorActions.cpp
--- a/Engine/Source/Editor/LevelEditor/Private/LevelEditorActions.cpp
+++ b/Engine/Source/Editor/LevelEditor/Private/LevelEditorActions.cpp
@@ -2078,7 +2078,7 @@ bool FLevelEditorActionCallbacks::SaveAnimationFromSkeletalMeshComponent(AActor
 		bool bSimulating = false;
 		for (auto & Comp : SimSkelComponents)
 		{
-			bSimulating |= (Comp->SkeletalMesh && Comp->SkeletalMesh->Skeleton && Comp->IsSimulatingPhysics());
+			bSimulating |= (Comp->SkeletalMesh && Comp->SkeletalMesh->GetSkeleton() && Comp->IsSimulatingPhysics());
 		}
 
 		// if any of them are legitimately simulating
@@ -2092,7 +2092,7 @@ bool FLevelEditorActionCallbacks::SaveAnimationFromSkeletalMeshComponent(AActor
 			{
 				for (auto & Comp : SimSkelComponents)
 				{
-					if (Comp->SkeletalMesh && Comp->SkeletalMesh->Skeleton && Comp->IsSimulatingPhysics())
+					if (Comp->SkeletalMesh && Comp->SkeletalMesh->GetSkeleton() && Comp->IsSimulatingPhysics())
 					{
 						// now record to animation
 						class UAnimSequence* Sequence = LevelEditorModule.OnCaptureSingleFrameAnimSequence().IsBound() ? LevelEditorModule.OnCaptureSingleFrameAnimSequence().Execute(Comp) : nullptr;
diff --git a/Engine/Source/Editor/Matinee/Private/MatineeTrackHelpers.cpp b/Engine/Source/Editor/Matinee/Private/MatineeTrackHelpers.cpp
--- a/Engine/Source/Editor/Matinee/Private/MatineeTrackHelpers.cpp
+++ b/Engine/Source/Editor/Matinee/Private/MatineeTrackHelpers.cpp
@@ -256,7 +256,7 @@ bool UMatineeTrackAnimControlHelper::PreCreateKeyframe( UInterpTrack *Track, flo
 		if (USkeletalMeshComponent* CurSkelMeshComp = Cast<USkeletalMeshComponent>(Component))
 		{
 			// if qualified to play animation, break
-			if (CurSkelMeshComp->SkeletalMesh && CurSkelMeshComp->SkeletalMesh->Skeleton)
+			if (CurSkelMeshComp->SkeletalMesh && CurSkelMeshComp->SkeletalMesh->GetSkeleton())
 			{
 				SkelMeshComp = CurSkelMeshComp;
 				break;
@@ -270,7 +270,7 @@ bool UMatineeTrackAnimControlHelper::PreCreateKeyframe( UInterpTrack *Track, flo
 		return false;
 	}
 
-	USkeleton* Skeleton = SkelMeshComp->SkeletalMesh->Skeleton;
+	USkeleton* Skeleton = SkelMeshComp->SkeletalMesh->GetSkeleton();
 	if ( Skeleton )
 	{
 		// Show the dialog.
diff --git a/Engine/Source/Editor/MeshPaint/Private/MeshPaintHelpers.cpp b/Engine/Source/Editor/MeshPaint/Private/MeshPaintHelpers.cpp
--- a/Engine/Source/Editor/MeshPaint/Private/MeshPaintHelpers.cpp
+++ b/Engine/Source/Editor/MeshPaint/Private/MeshPaintHelpers.cpp
@@ -839,8 +839,8 @@ void MeshPaintHelpers::FillSkeletalMeshVertexColors(USkeletalMeshComponent* Mesh
 		// Dirty the mesh
 		Mesh->SetFlags(RF_Transactional);
 		Mesh->Modify();
-		Mesh->bHasVertexColors = true;
-		Mesh->VertexColorGuid = FGuid::NewGuid();
+		Mesh->SetHasVertexColors(true);
+		Mesh->SetVertexColorGuid(FGuid::NewGuid());
 
 		// Release the static mesh's resources.
 		Mesh->ReleaseResources();
@@ -1530,9 +1530,9 @@ void MeshPaintHelpers::ImportVertexColorsToSkeletalMesh(USkeletalMesh* SkeletalM
 	}
 
 	//Make sure we change the import data so the re-import do not replace the new data
-	if (SkeletalMesh->AssetImportData)
+	if (SkeletalMesh->GetAssetImportData())
 	{
-		UFbxSkeletalMeshImportData* ImportData = Cast<UFbxSkeletalMeshImportData>(SkeletalMesh->AssetImportData);
+		UFbxSkeletalMeshImportData* ImportData = Cast<UFbxSkeletalMeshImportData>(SkeletalMesh->GetAssetImportData());
 		if (ImportData && ImportData->VertexColorImportOption != EVertexColorImportOption::Ignore)
 		{
 			ImportData->SetFlags(RF_Transactional);
diff --git a/Engine/Source/Editor/MeshPaint/Private/MeshPaintSkeletalMeshAdapter.cpp b/Engine/Source/Editor/MeshPaint/Private/MeshPaintSkeletalMeshAdapter.cpp
--- a/Engine/Source/Editor/MeshPaint/Private/MeshPaintSkeletalMeshAdapter.cpp
+++ b/Engine/Source/Editor/MeshPaint/Private/MeshPaintSkeletalMeshAdapter.cpp
@@ -201,7 +201,7 @@ void FMeshPaintGeometryAdapterForSkeletalMeshes::CleanupGlobals()
 	{
 		if (Pair.Key && Pair.Value.RestoreBodySetup)
 		{
-			Pair.Key->BodySetup = Pair.Value.RestoreBodySetup;
+			Pair.Key->SetBodySetup(Pair.Value.RestoreBodySetup);
 		}
 	}
 
@@ -226,13 +226,14 @@ void FMeshPaintGeometryAdapterForSkeletalMeshes::OnAdded()
 	if (SkeletalMeshReferencers.Referencers.Num() == 0)
 	{
 		// Remember the old body setup (this will be added as a GC reference so that it doesn't get destroyed)
-		SkeletalMeshReferencers.RestoreBodySetup = ReferencedSkeletalMesh->BodySetup;
+		const USkeletalMesh* ReferencedSkeletalMeshConst = ReferencedSkeletalMesh;
+		SkeletalMeshReferencers.RestoreBodySetup = ReferencedSkeletalMeshConst->GetBodySetup();
 
 		if (SkeletalMeshReferencers.RestoreBodySetup)
 		{
 			// Create a new body setup from the mesh's main body setup. This has to have the skeletal mesh as its outer,
 			// otherwise the body instance will not be created correctly.
-			UBodySetup* TempBodySetupRaw = DuplicateObject<UBodySetup>(ReferencedSkeletalMesh->BodySetup, ReferencedSkeletalMesh);
+			UBodySetup* TempBodySetupRaw = DuplicateObject<UBodySetup>(ReferencedSkeletalMeshConst->GetBodySetup(), ReferencedSkeletalMesh);
 			TempBodySetupRaw->ClearFlags(RF_Transactional);
 
 			// Set collide all flag so that the body creates physics meshes using ALL elements from the mesh not just the collision mesh.
@@ -246,7 +247,7 @@ void FMeshPaintGeometryAdapterForSkeletalMeshes::OnAdded()
 			TempBodySetupRaw->AggGeom.ConvexElems.Empty();
 
 			// Set as new body setup
-			ReferencedSkeletalMesh->BodySetup = TempBodySetupRaw;
+			ReferencedSkeletalMesh->SetBodySetup(TempBodySetupRaw);
 		}
 	}
 
@@ -303,7 +304,7 @@ void FMeshPaintGeometryAdapterForSkeletalMeshes::OnRemoved()
 	{
 		if (SkeletalMeshReferencers->RestoreBodySetup != nullptr)
 		{
-			ReferencedSkeletalMesh->BodySetup = SkeletalMeshReferencers->RestoreBodySetup;
+			ReferencedSkeletalMesh->SetBodySetup(SkeletalMeshReferencers->RestoreBodySetup);
 		}
 		
 		verify(MeshToComponentMap.Remove(ReferencedSkeletalMesh) == 1);
@@ -437,8 +438,8 @@ void FMeshPaintGeometryAdapterForSkeletalMeshes::PreEdit()
 	ReferencedSkeletalMesh->SetFlags(RF_Transactional);
 	ReferencedSkeletalMesh->Modify();
 
-	ReferencedSkeletalMesh->bHasVertexColors = true;
-	ReferencedSkeletalMesh->VertexColorGuid = FGuid::NewGuid();
+	ReferencedSkeletalMesh->SetHasVertexColors(true);
+	ReferencedSkeletalMesh->SetVertexColorGuid(FGuid::NewGuid());
 
 	// Release the static mesh's resources.
 	ReferencedSkeletalMesh->ReleaseResources();
@@ -451,14 +452,14 @@ void FMeshPaintGeometryAdapterForSkeletalMeshes::PreEdit()
 	{
 		// Mesh doesn't have a color vertex buffer yet!  We'll create one now.
 		LODData->StaticVertexBuffers.ColorVertexBuffer.InitFromSingleColor(FColor(255, 255, 255, 255), LODData->GetNumVertices());
-		ReferencedSkeletalMesh->bHasVertexColors = true;
-		ReferencedSkeletalMesh->VertexColorGuid = FGuid::NewGuid();
+		ReferencedSkeletalMesh->SetHasVertexColors(true);
+		ReferencedSkeletalMesh->SetVertexColorGuid(FGuid::NewGuid());
 		BeginInitResource(&LODData->StaticVertexBuffers.ColorVertexBuffer);
 	}
 	//Make sure we change the import data so the re-import do not replace the new data
-	if (ReferencedSkeletalMesh->AssetImportData)
+	if (ReferencedSkeletalMesh->GetAssetImportData())
 	{
-		UFbxSkeletalMeshImportData* ImportData = Cast<UFbxSkeletalMeshImportData>(ReferencedSkeletalMesh->AssetImportData);
+		UFbxSkeletalMeshImportData* ImportData = Cast<UFbxSkeletalMeshImportData>(ReferencedSkeletalMesh->GetAssetImportData());
 		if (ImportData && ImportData->VertexColorImportOption != EVertexColorImportOption::Ignore)
 		{
 			ImportData->SetFlags(RF_Transactional);
diff --git a/Engine/Source/Editor/MovieSceneTools/Private/EditModes/SkeletalAnimationTrackEditMode.cpp b/Engine/Source/Editor/MovieSceneTools/Private/EditModes/SkeletalAnimationTrackEditMode.cpp
--- a/Engine/Source/Editor/MovieSceneTools/Private/EditModes/SkeletalAnimationTrackEditMode.cpp
+++ b/Engine/Source/Editor/MovieSceneTools/Private/EditModes/SkeletalAnimationTrackEditMode.cpp
@@ -243,7 +243,7 @@ static void DrawBonesFromCompactPose(const FCompactPose& Pose, USkeletalMeshComp
 
 		if (MeshComponent && MeshComponent->SkeletalMesh)
 		{
-			DrawBones(Pose.GetBoneContainer().GetBoneIndicesArray(), MeshComponent->SkeletalMesh->RefSkeleton, WorldTransforms,
+			DrawBones(Pose.GetBoneContainer().GetBoneIndicesArray(), MeshComponent->SkeletalMesh->GetRefSkeleton(), WorldTransforms,
 				LODData, PDI, BoneColours, MeshComponent->Bounds.SphereRadius, 1.0f);
 		}
 
diff --git a/Engine/Source/Editor/MovieSceneTools/Private/TrackEditors/SkeletalAnimationTrackEditor.cpp b/Engine/Source/Editor/MovieSceneTools/Private/TrackEditors/SkeletalAnimationTrackEditor.cpp
--- a/Engine/Source/Editor/MovieSceneTools/Private/TrackEditors/SkeletalAnimationTrackEditor.cpp
+++ b/Engine/Source/Editor/MovieSceneTools/Private/TrackEditors/SkeletalAnimationTrackEditor.cpp
@@ -316,10 +316,10 @@ USkeletalMeshComponent* AcquireSkeletalMeshFromObjectGuid(const FGuid& Guid, TSh
 USkeleton* GetSkeletonFromComponent(UActorComponent* InComponent)
 {
 	USkeletalMeshComponent* SkeletalMeshComp = Cast<USkeletalMeshComponent>(InComponent);
-	if (SkeletalMeshComp && SkeletalMeshComp->SkeletalMesh && SkeletalMeshComp->SkeletalMesh->Skeleton)
+	if (SkeletalMeshComp && SkeletalMeshComp->SkeletalMesh && SkeletalMeshComp->SkeletalMesh->GetSkeleton())
 	{
 		// @todo Multiple actors, multiple components
-		return SkeletalMeshComp->SkeletalMesh->Skeleton;
+		return SkeletalMeshComp->SkeletalMesh->GetSkeleton();
 	}
 
 	return nullptr;
diff --git a/Engine/Source/Editor/Persona/Private/AnimGraphNodeDetails.cpp b/Engine/Source/Editor/Persona/Private/AnimGraphNodeDetails.cpp
--- a/Engine/Source/Editor/Persona/Private/AnimGraphNodeDetails.cpp
+++ b/Engine/Source/Editor/Persona/Private/AnimGraphNodeDetails.cpp
@@ -521,7 +521,7 @@ void FBoneReferenceCustomization::SetEditableSkeleton(TSharedRef<IPropertyHandle
 
 		if (USkeletalMesh* SkeletalMesh = Cast<USkeletalMesh>(Outer))
 		{
-			TargetSkeleton = SkeletalMesh->Skeleton;
+			TargetSkeleton = SkeletalMesh->GetSkeleton();
 			break;
 		}
 
@@ -529,7 +529,7 @@ void FBoneReferenceCustomization::SetEditableSkeleton(TSharedRef<IPropertyHandle
 		{
 			USkeletalMesh* SkeletalMesh = LODInfoUILayout->GetPersonaToolkit()->GetPreviewMesh();
 			check(SkeletalMesh);
-			TargetSkeleton = SkeletalMesh->Skeleton;
+			TargetSkeleton = SkeletalMesh->GetSkeleton();
 			break;
 		}
 
diff --git a/Engine/Source/Editor/Persona/Private/AnimationEditorPreviewScene.cpp b/Engine/Source/Editor/Persona/Private/AnimationEditorPreviewScene.cpp
--- a/Engine/Source/Editor/Persona/Private/AnimationEditorPreviewScene.cpp
+++ b/Engine/Source/Editor/Persona/Private/AnimationEditorPreviewScene.cpp
@@ -405,9 +405,10 @@ void FAnimationEditorPreviewScene::AddPreviewAttachedObjects()
 
 	if ( Mesh )
 	{
-		for(int32 i = 0; i < Mesh->PreviewAttachedAssetContainer.Num(); i++)
+		FPreviewAssetAttachContainer& PreviewAssetAttachContainer = Mesh->GetPreviewAttachedAssetContainer();
+		for(int32 Index = 0; Index < PreviewAssetAttachContainer.Num(); Index++)
 		{
-			FPreviewAttachedObjectPair& PreviewAttachedObject = Mesh->PreviewAttachedAssetContainer[i];
+			FPreviewAttachedObjectPair& PreviewAttachedObject = PreviewAssetAttachContainer[Index];
 
 			AttachObjectToPreviewComponent(PreviewAttachedObject.GetAttachedObject(), PreviewAttachedObject.AttachedTo);
 		}
@@ -521,7 +522,7 @@ void FAnimationEditorPreviewScene::RemoveAttachedComponent( bool bRemovePreviewA
 
 		if ( USkeletalMesh* PreviewMesh = PersonaToolkit.Pin()->GetMesh() )
 		{
-			for(auto Iter = PreviewMesh->PreviewAttachedAssetContainer.CreateConstIterator(); Iter; ++Iter)
+			for(auto Iter = PreviewMesh->GetPreviewAttachedAssetContainer().CreateConstIterator(); Iter; ++Iter)
 			{
 				PreviewAttachedObjects.FindOrAdd(Iter->GetAttachedObject()).Add(Iter->AttachedTo);
 			}
@@ -655,10 +656,10 @@ void FAnimationEditorPreviewScene::ShowReferencePose(bool bShowRefPose, bool bRe
 			bool bModified = false;
 			FScopedTransaction Transaction(LOCTEXT("ResetBoneTransforms", "Reset Bone Transforms"));
 
-			int32 NumBones = SkeletalMeshComponent->SkeletalMesh->RefSkeleton.GetNum();
+			int32 NumBones = SkeletalMeshComponent->SkeletalMesh->GetRefSkeleton().GetNum();
 			for (int32 BoneIndex = 0; BoneIndex < NumBones; ++BoneIndex)
 			{
-				FName BoneName = SkeletalMeshComponent->SkeletalMesh->RefSkeleton.GetBoneName(BoneIndex);
+				FName BoneName = SkeletalMeshComponent->SkeletalMesh->GetRefSkeleton().GetBoneName(BoneIndex);
 				const FAnimNode_ModifyBone* ModifiedBone = SkeletalMeshComponent->PreviewInstance->FindModifiedBone(BoneName);
 				if (ModifiedBone != nullptr)
 				{
diff --git a/Engine/Source/Editor/Persona/Private/AnimationEditorViewportClient.cpp b/Engine/Source/Editor/Persona/Private/AnimationEditorViewportClient.cpp
--- a/Engine/Source/Editor/Persona/Private/AnimationEditorViewportClient.cpp
+++ b/Engine/Source/Editor/Persona/Private/AnimationEditorViewportClient.cpp
@@ -306,11 +306,11 @@ void FAnimationViewportClient::SaveCameraAsDefault()
 
 		FViewportCameraTransform& ViewTransform = GetViewTransform();
 		SkelMesh->Modify();
-		SkelMesh->DefaultEditorCameraLocation = ViewTransform.GetLocation();
-		SkelMesh->DefaultEditorCameraRotation = ViewTransform.GetRotation();
-		SkelMesh->DefaultEditorCameraLookAt = ViewTransform.GetLookAt();
-		SkelMesh->DefaultEditorCameraOrthoZoom = ViewTransform.GetOrthoZoom();
-		SkelMesh->bHasCustomDefaultEditorCamera = true;
+		SkelMesh->SetDefaultEditorCameraLocation(ViewTransform.GetLocation());
+		SkelMesh->SetDefaultEditorCameraRotation(ViewTransform.GetRotation());
+		SkelMesh->SetDefaultEditorCameraLookAt(ViewTransform.GetLookAt());
+		SkelMesh->SetDefaultEditorCameraOrthoZoom(ViewTransform.GetOrthoZoom());
+		SkelMesh->SetHasCustomDefaultEditorCamera(true);
 
 		// Create and display a notification 
 		const FText NotificationText = FText::Format(LOCTEXT("SavedDefaultCamera", "Saved default camera for {0}"), FText::AsCultureInvariant(SkelMesh->GetName()));
@@ -333,7 +333,7 @@ void FAnimationViewportClient::ClearDefaultCamera()
 		FScopedTransaction Transaction(LOCTEXT("ClearDefaultCamera", "Clear Default Camera"));
 
 		SkelMesh->Modify();
-		SkelMesh->bHasCustomDefaultEditorCamera = false;
+		SkelMesh->SetHasCustomDefaultEditorCamera(false);
 
 		// Create and display a notification 
 		const FText NotificationText = FText::Format(LOCTEXT("ClearedDefaultCamera", "Cleared default camera for {0}"), FText::AsCultureInvariant(SkelMesh->GetName()));
@@ -346,7 +346,7 @@ void FAnimationViewportClient::ClearDefaultCamera()
 bool FAnimationViewportClient::HasDefaultCameraSet() const
 {
 	USkeletalMesh* SkelMesh = GetAnimPreviewScene()->GetPreviewMeshComponent()->SkeletalMesh;
-	return (SkelMesh && SkelMesh->bHasCustomDefaultEditorCamera);
+	return (SkelMesh && SkelMesh->GetHasCustomDefaultEditorCamera());
 }
 
 void FAnimationViewportClient::HandleSkeletalMeshChanged(USkeletalMesh* OldSkeletalMesh, USkeletalMesh* NewSkeletalMesh)
@@ -429,9 +429,9 @@ void FAnimationViewportClient::Draw(const FSceneView* View, FPrimitiveDrawInterf
 		// Display socket hit points
 		if (PreviewMeshComponent->bDrawSockets )
 		{
-			if (PreviewMeshComponent->bSkeletonSocketsVisible && PreviewMeshComponent->SkeletalMesh->Skeleton )
+			if (PreviewMeshComponent->bSkeletonSocketsVisible && PreviewMeshComponent->SkeletalMesh->GetSkeleton() )
 			{
-				DrawSockets(PreviewMeshComponent, PreviewMeshComponent->SkeletalMesh->Skeleton->Sockets, FSelectedSocketInfo(), PDI, true);
+				DrawSockets(PreviewMeshComponent, PreviewMeshComponent->SkeletalMesh->GetSkeleton()->Sockets, FSelectedSocketInfo(), PDI, true);
 			}
 
 			if ( PreviewMeshComponent->bMeshSocketsVisible )
@@ -705,7 +705,7 @@ FText FAnimationViewportClient::GetDisplayInfo(bool bDisplayAllInfo) const
 		TextValue = ConcatenateLine(TextValue, LOCTEXT("AdditiveRefPoseWarning", "<AnimViewport.WarningText>Additive ref pose contains scales of 0.0, this can cause additive animations to not give the desired results</>"));
 	}
 
-	if (PreviewMeshComponent->SkeletalMesh->MorphTargets.Num() > 0)
+	if (PreviewMeshComponent->SkeletalMesh->GetMorphTargets().Num() > 0)
 	{
 		TArray<UMaterial*> ProcessedMaterials;
 		TArray<UMaterial*> MaterialsThatNeedMorphFlagOn;
@@ -714,7 +714,7 @@ FText FAnimationViewportClient::GetDisplayInfo(bool bDisplayAllInfo) const
 		const TIndirectArray<FSkeletalMeshLODModel>& LODModels = PreviewMeshComponent->SkeletalMesh->GetImportedModel()->LODModels;
 		int32 LodNumber = LODModels.Num();
 		TArray<UMaterialInterface*> MaterialUsingMorphTarget;
-		for (UMorphTarget *MorphTarget : PreviewMeshComponent->SkeletalMesh->MorphTargets)
+		for (UMorphTarget *MorphTarget : PreviewMeshComponent->SkeletalMesh->GetMorphTargets())
 		{
 			if (MorphTarget == nullptr)
 			{
@@ -729,7 +729,7 @@ FText FAnimationViewportClient::GetDisplayInfo(bool bDisplayAllInfo) const
 						const FSkeletalMeshLODModel& LODModel = LODModels[LodIdx];
 						if (LODModel.Sections.IsValidIndex(SectionIndex))
 						{
-							MaterialUsingMorphTarget.AddUnique(PreviewMeshComponent->SkeletalMesh->Materials[LODModel.Sections[SectionIndex].MaterialIndex].MaterialInterface);
+							MaterialUsingMorphTarget.AddUnique(PreviewMeshComponent->SkeletalMesh->GetMaterials()[LODModel.Sections[SectionIndex].MaterialIndex].MaterialInterface);
 						}
 					}
 				}
@@ -1673,14 +1673,14 @@ void FAnimationViewportClient::FocusViewportOnPreviewMesh(bool bUseCustomCamera)
 	{
 		if (USkeletalMesh* const SkelMesh = PreviewMeshComponent->SkeletalMesh)
 		{
-			if (bUseCustomCamera && SkelMesh->bHasCustomDefaultEditorCamera)
+			if (bUseCustomCamera && SkelMesh->GetHasCustomDefaultEditorCamera())
 			{
 				FViewportCameraTransform& ViewTransform = GetViewTransform();
 
-				ViewTransform.SetLocation(SkelMesh->DefaultEditorCameraLocation);
-				ViewTransform.SetRotation(SkelMesh->DefaultEditorCameraRotation);
-				ViewTransform.SetLookAt(SkelMesh->DefaultEditorCameraLookAt);
-				ViewTransform.SetOrthoZoom(SkelMesh->DefaultEditorCameraOrthoZoom);
+				ViewTransform.SetLocation(SkelMesh->GetDefaultEditorCameraLocation());
+				ViewTransform.SetRotation(SkelMesh->GetDefaultEditorCameraRotation());
+				ViewTransform.SetLookAt(SkelMesh->GetDefaultEditorCameraLookAt());
+				ViewTransform.SetOrthoZoom(SkelMesh->GetDefaultEditorCameraOrthoZoom());
 
 				Invalidate();
 				return;
@@ -1709,7 +1709,7 @@ float FAnimationViewportClient::GetFloorOffset() const
 	USkeletalMesh* Mesh = GetPreviewScene()->GetPreviewMeshComponent()->SkeletalMesh;
 	if ( Mesh )
 	{
-		return Mesh->FloorOffset;
+		return Mesh->GetFloorOffset();
 	}
 
 	return 0.0f;
@@ -1725,7 +1725,7 @@ void FAnimationViewportClient::SetFloorOffset( float NewValue )
 		FScopedTransaction Transaction( LOCTEXT( "SetFloorOffset", "Set Floor Offset" ) );
 		Mesh->Modify();
 
-		Mesh->FloorOffset = NewValue;
+		Mesh->SetFloorOffset(NewValue);
 		UpdateCameraSetup(); // This does the actual moving of the floor mesh
 		Invalidate();
 	}
diff --git a/Engine/Source/Editor/Persona/Private/AssetSearchBoxUtilPersona.cpp b/Engine/Source/Editor/Persona/Private/AssetSearchBoxUtilPersona.cpp
--- a/Engine/Source/Editor/Persona/Private/AssetSearchBoxUtilPersona.cpp
+++ b/Engine/Source/Editor/Persona/Private/AssetSearchBoxUtilPersona.cpp
@@ -44,7 +44,7 @@ void SAssetSearchBoxForBones::Construct( const FArguments& InArgs, const class U
 				PossibleSuggestions.Add(FAssetSearchBoxSuggestion::MakeSimpleSuggestion((*SocketIt)->SocketName.ToString()));
 			}
 		}
-		Skeleton = SkeletalMesh->Skeleton;
+		Skeleton = SkeletalMesh->GetSkeleton();
 	}
 	else
 	{
diff --git a/Engine/Source/Editor/Persona/Private/Customization/SkeletalMeshRegionCustomization.cpp b/Engine/Source/Editor/Persona/Private/Customization/SkeletalMeshRegionCustomization.cpp
--- a/Engine/Source/Editor/Persona/Private/Customization/SkeletalMeshRegionCustomization.cpp
+++ b/Engine/Source/Editor/Persona/Private/Customization/SkeletalMeshRegionCustomization.cpp
@@ -105,10 +105,10 @@ void FNiagaraSkeletalMeshRegionBoneFilterDetails::HandleBoneNameComboBoxSelectio
 void FNiagaraSkeletalMeshRegionBoneFilterDetails::OnComboOpening()
 {
 	PossibleBoneNames.Reset();
-	if (MeshObject->Skeleton != nullptr)
+	if (MeshObject->GetSkeleton() != nullptr)
 	{
 		//Populate PossibleBonesNames
-		for (FMeshBoneInfo Bone : MeshObject->Skeleton->GetReferenceSkeleton().GetRefBoneInfo())
+		for (FMeshBoneInfo Bone : MeshObject->GetSkeleton()->GetReferenceSkeleton().GetRefBoneInfo())
 		{
 			TSharedPtr<FName> BoneName(new FName(Bone.Name));
 			PossibleBoneNames.Add(BoneName);
@@ -195,10 +195,10 @@ void FNiagaraSkeletalMeshRegionMaterialFilterDetails::HandleMaterialNameComboBox
 void FNiagaraSkeletalMeshRegionMaterialFilterDetails::OnComboOpening()
 {
 	PossibleMaterialNames.Reset();
-	if (MeshObject->Skeleton != nullptr)
+	if (MeshObject->GetSkeleton() != nullptr)
 	{
 		//Populate PossibleMaterialsNames
-		for (FSkeletalMaterial Material : MeshObject->Materials)
+		for (FSkeletalMaterial Material : MeshObject->GetMaterials())
 		{
 			TSharedPtr<FName> MaterialName(new FName(Material.MaterialSlotName));
 			PossibleMaterialNames.Add(MaterialName);
diff --git a/Engine/Source/Editor/Persona/Private/EditModes/SkeletonSelectionEditMode.cpp b/Engine/Source/Editor/Persona/Private/EditModes/SkeletonSelectionEditMode.cpp
--- a/Engine/Source/Editor/Persona/Private/EditModes/SkeletonSelectionEditMode.cpp
+++ b/Engine/Source/Editor/Persona/Private/EditModes/SkeletonSelectionEditMode.cpp
@@ -42,7 +42,7 @@ bool FSkeletonSelectionEditMode::GetCameraTarget(FSphere& OutTarget) const
 		const int32 FocusBoneIndex = GetAnimPreviewScene().GetSelectedBoneIndex();
 		if (FocusBoneIndex != INDEX_NONE)
 		{
-			const FName BoneName = PreviewMeshComponent->SkeletalMesh->RefSkeleton.GetBoneName(FocusBoneIndex);
+			const FName BoneName = PreviewMeshComponent->SkeletalMesh->GetRefSkeleton().GetBoneName(FocusBoneIndex);
 			OutTarget.Center = PreviewMeshComponent->GetBoneLocation(BoneName);
 			OutTarget.W = 30.0f;
 			bHandled = true;
@@ -151,7 +151,7 @@ bool FSkeletonSelectionEditMode::StartTracking(FEditorViewportClient* InViewport
 					bInTransaction = true;
 
 					// now modify the bone array
-					const FName BoneName = PreviewMeshComponent->SkeletalMesh->RefSkeleton.GetBoneName(BoneIndex);
+					const FName BoneName = PreviewMeshComponent->SkeletalMesh->GetRefSkeleton().GetBoneName(BoneIndex);
 					PreviewMeshComponent->PreviewInstance->ModifyBone(BoneName);
 				}
 			}
@@ -209,7 +209,7 @@ bool FSkeletonSelectionEditMode::InputDelta(FEditorViewportClient* InViewportCli
 		if ( BoneIndex >= 0 )
 		{
 			//Get the skeleton control manipulating this bone
-			const FName BoneName = PreviewMeshComponent->SkeletalMesh->RefSkeleton.GetBoneName(BoneIndex);
+			const FName BoneName = PreviewMeshComponent->SkeletalMesh->GetRefSkeleton().GetBoneName(BoneIndex);
 			SkelControl = &(PreviewMeshComponent->PreviewInstance->ModifyBone(BoneName));
 		}
 
@@ -352,7 +352,7 @@ void FSkeletonSelectionEditMode::DrawHUD(FEditorViewportClient* ViewportClient,
 		const int32 HalfY = ViewPortSize.Y / 2;
 
 		int32 BoneIndex = GetAnimPreviewScene().GetSelectedBoneIndex();
-		const FName BoneName = PreviewMeshComponent->SkeletalMesh->RefSkeleton.GetBoneName(BoneIndex);
+		const FName BoneName = PreviewMeshComponent->SkeletalMesh->GetRefSkeleton().GetBoneName(BoneIndex);
 
 		FMatrix BoneMatrix = PreviewMeshComponent->GetBoneMatrix(BoneIndex);
 		const FPlane Proj = View->Project(BoneMatrix.GetOrigin());
@@ -480,7 +480,7 @@ FVector FSkeletonSelectionEditMode::GetWidgetLocation() const
 	{
 		if (PreviewMeshComponent && PreviewMeshComponent->SkeletalMesh)
 		{
-			const FName BoneName = PreviewMeshComponent->SkeletalMesh->RefSkeleton.GetBoneName(BoneIndex);
+			const FName BoneName = PreviewMeshComponent->SkeletalMesh->GetRefSkeleton().GetBoneName(BoneIndex);
 
 			const FMatrix BoneMatrix = PreviewMeshComponent->GetBoneMatrix(BoneIndex);
 
diff --git a/Engine/Source/Editor/Persona/Private/PersonaAssetFamily.cpp b/Engine/Source/Editor/Persona/Private/PersonaAssetFamily.cpp
--- a/Engine/Source/Editor/Persona/Private/PersonaAssetFamily.cpp
+++ b/Engine/Source/Editor/Persona/Private/PersonaAssetFamily.cpp
@@ -358,7 +358,7 @@ void FPersonaAssetFamily::FindCounterpartAssets(const UObject* InAsset, const US
 	else if (InAsset->IsA<USkeletalMesh>())
 	{
 		OutMesh = CastChecked<USkeletalMesh>(InAsset);
-		OutSkeleton = OutMesh->Skeleton;
+		OutSkeleton = OutMesh->GetSkeleton();
 	}
 	else if (InAsset->IsA<UAnimBlueprint>())
 	{
@@ -381,7 +381,7 @@ void FPersonaAssetFamily::FindCounterpartAssets(const UObject* InAsset, const US
 		OutMesh = PhysicsAsset->PreviewSkeletalMesh.LoadSynchronous();
 		if(OutMesh != nullptr)
 		{
-			OutSkeleton = OutMesh->Skeleton;
+			OutSkeleton = OutMesh->GetSkeleton();
 		}
 	}
 }
