Add missing FGCScopeGuard when creating MorphTarget which can potentially run from another thread when async skeletal mesh feature is activated

diff --git a/Engine/Source/Developer/SkeletalMeshUtilitiesCommon/Private/LODUtilities.cpp b/Engine/Source/Developer/SkeletalMeshUtilitiesCommon/Private/LODUtilities.cpp
--- a/Engine/Source/Developer/SkeletalMeshUtilitiesCommon/Private/LODUtilities.cpp
+++ b/Engine/Source/Developer/SkeletalMeshUtilitiesCommon/Private/LODUtilities.cpp
@@ -8,6 +8,7 @@
 #include "Components/SkinnedMeshComponent.h"
 #include "Components/SkeletalMeshComponent.h"
 #include "Animation/MorphTarget.h"
+#include "UObject/GarbageCollection.h"
 #include "Rendering/SkeletalMeshModel.h"
 #include "Rendering/SkeletalMeshLODModel.h"
 #include "GenericQuadTree.h"
@@ -2821,6 +2822,7 @@ void FLODUtilities::BuildMorphTargets(USkeletalMesh* BaseSkelMesh, FSkeletalMesh
 			{
 				if (LODIndex == 0)
 				{
+					FGCScopeGuard GCScopeGuard;
 					MorphTarget = NewObject<UMorphTarget>(BaseSkelMesh, ObjectName);
 				}
 				else
