Fix re-import of morph target crash when asynchronous skeletal mesh build is enable.

diff --git a/Engine/Source/Developer/SkeletalMeshUtilitiesCommon/Private/LODUtilities.cpp b/Engine/Source/Developer/SkeletalMeshUtilitiesCommon/Private/LODUtilities.cpp
--- a/Engine/Source/Developer/SkeletalMeshUtilitiesCommon/Private/LODUtilities.cpp
+++ b/Engine/Source/Developer/SkeletalMeshUtilitiesCommon/Private/LODUtilities.cpp
@@ -2812,6 +2812,7 @@ void FLODUtilities::BuildMorphTargets(USkeletalMesh* BaseSkelMesh, FSkeletalMesh
 		FSkeletalMeshImportData& ShapeImportData = BaseImportData.MorphTargets[MorphTargetIndex];
 		TSet<uint32>& ModifiedPoints = BaseImportData.MorphTargetModifiedPoints[MorphTargetIndex];
 
+		bool bNeedToClearAsyncFlag = false;
 		UMorphTarget* MorphTarget = nullptr;
 		{
 			FName ObjectName = *ShapeName;
@@ -2822,6 +2823,21 @@ void FLODUtilities::BuildMorphTargets(USkeletalMesh* BaseSkelMesh, FSkeletalMesh
 			{
 				if (LODIndex == 0)
 				{
+					if (!IsInGameThread())
+					{
+						//TODO remove this code when overriding a UObject will be allow outside of the game thread
+						//We currently need to avoid overriding an existing asset outside of the game thread
+						UObject* ExistingMorphTarget = StaticFindObject(UMorphTarget::StaticClass(), BaseSkelMesh, *ShapeName);
+						if (ExistingMorphTarget)
+						{
+							//make sure the object is not standalone or transactional
+							ExistingMorphTarget->ClearFlags(RF_Standalone | RF_Transactional);
+							//Move this object in the transient package
+							ExistingMorphTarget->Rename(nullptr, GetTransientPackage(), REN_ForceNoResetLoaders | REN_DoNotDirty | REN_DontCreateRedirectors | REN_NonTransactional);
+							ExistingMorphTarget = nullptr;
+						}
+						bNeedToClearAsyncFlag = true;
+					}
 					FGCScopeGuard GCScopeGuard;
 					MorphTarget = NewObject<UMorphTarget>(BaseSkelMesh, ObjectName);
 				}
@@ -2849,6 +2865,12 @@ void FLODUtilities::BuildMorphTargets(USkeletalMesh* BaseSkelMesh, FSkeletalMesh
 			NewWork->StartBackgroundTask(GLargeThreadPool);
 			CurrentNumTasks++;
 			NumTasks++;
+
+			if (bNeedToClearAsyncFlag)
+			{
+				const EInternalObjectFlags AsyncFlags = EInternalObjectFlags::Async | EInternalObjectFlags::AsyncLoading;
+				MorphTarget->ClearInternalFlags(AsyncFlags);
+			}
 		}
 
 		++ShapeIndex;
