Get rid of huge amount of allocs which affects Linux more than Windows due to the allocator used

#rnx

diff --git a/Engine/Source/Developer/SkeletalMeshUtilitiesCommon/Private/LODUtilities.cpp b/Engine/Source/Developer/SkeletalMeshUtilitiesCommon/Private/LODUtilities.cpp
--- a/Engine/Source/Developer/SkeletalMeshUtilitiesCommon/Private/LODUtilities.cpp
+++ b/Engine/Source/Developer/SkeletalMeshUtilitiesCommon/Private/LODUtilities.cpp
@@ -2520,6 +2520,7 @@ public:
 		WasProcessed.AddZeroed(MeshDataBundle.Indices.Num());
 
 		TArray< int32 > WedgeFaces;
+		TArray< int32 > OtherWedgeFaces;
 		TArray< int32 > OverlappingWedgesDummy;
 		TArray< int32 > OtherOverlappingWedgesDummy;
 
@@ -2558,7 +2559,7 @@ public:
 
 							for (const int32 OtherDupVert : OtherOverlappingWedges)
 							{
-								TArray< int32 > OtherWedgeFaces;
+								OtherWedgeFaces.Reset();
 								WedgeToFaces.MultiFind(OtherDupVert, OtherWedgeFaces);
 
 								for (const int32 OtherFaceIndex : OtherWedgeFaces)
