Remove original reduction data. Use the source data and the build instead.

#rnx

diff --git a/Engine/Plugins/Experimental/SkeletalReduction/Source/Private/SkeletalMeshReductionPlugin.cpp b/Engine/Plugins/Experimental/SkeletalReduction/Source/Private/SkeletalMeshReductionPlugin.cpp
--- a/Engine/Plugins/Experimental/SkeletalReduction/Source/Private/SkeletalMeshReductionPlugin.cpp
+++ b/Engine/Plugins/Experimental/SkeletalReduction/Source/Private/SkeletalMeshReductionPlugin.cpp
@@ -1872,7 +1872,8 @@ void FQuadricSkeletalMeshReduction::ReduceSkeletalMesh(USkeletalMesh& SkeletalMe
 	// only allow to set BaseLOD if the LOD is less than this
 	if (Settings.BaseLOD > 0)
 	{
-		if (Settings.BaseLOD == LODIndex && (!SkelResource->OriginalReductionSourceMeshData.IsValidIndex(Settings.BaseLOD) || SkelResource->OriginalReductionSourceMeshData[Settings.BaseLOD]->IsEmpty()))
+		
+		if (Settings.BaseLOD == LODIndex && !SkeletalMesh.IsLODImportedDataBuildAvailable(LODIndex))
 		{
 			//Cannot reduce ourself if we are not imported
 			UE_LOG(LogSkeletalMeshReduction, Warning, TEXT("Building LOD %d - Cannot generate LOD with himself if the LOD do not have imported Data. Using Base LOD 0 instead"), LODIndex);
@@ -1940,17 +1941,9 @@ void FQuadricSkeletalMeshReduction::ReduceSkeletalMesh(USkeletalMesh& SkeletalMe
 	}
 
 	bool bReducingSourceModel = false;
-	//Reducing source LOD, we need to use the temporary data so it can be iterative
-	if (BaseLOD == LODIndex && (SkelResource->OriginalReductionSourceMeshData.IsValidIndex(BaseLOD) && !SkelResource->OriginalReductionSourceMeshData[BaseLOD]->IsEmpty()))
+	//Reducing source LOD, we need to use the LOD import data so it can be iterative
+	if (BaseLOD == LODIndex && SkeletalMesh.IsLODImportedDataBuildAvailable(LODIndex))
 	{
-		TMap<FString, TArray<FMorphTargetDelta>> TempLODMorphTargetData;
-		SkelResource->OriginalReductionSourceMeshData[BaseLOD]->LoadReductionData(*SrcModel, TempLODMorphTargetData, &SkeletalMesh);
-
-		//Restore the section data state (like disabled...)
-		RestoreUserSectionsData(SrcModelBackup, *SrcModel);
-		//Rebackup the source model since we update it, source always have empty LODMaterial map
-		//If you swap a material ID and after you do inline reduction, you have to remap it again, but not if you reduce and then remap the materialID
-		//this is by design currently
 		bReducingSourceModel = true;
 	}
 	else
@@ -2008,7 +2001,7 @@ void FQuadricSkeletalMeshReduction::ReduceSkeletalMesh(USkeletalMesh& SkeletalMe
 			RequiredBoneIndexArray[BoneIndex] = BoneIndex;
 		}
 
-		FBoneContainer RequiredBones(RequiredBoneIndexArray, false, *(&SkeletalMesh));
+		FBoneContainer RequiredBones(RequiredBoneIndexArray, false, SkeletalMesh);
 		RequiredBones.SetUseRAWData(true);
 
 		FCompactPose Pose;
@@ -2146,32 +2139,7 @@ void FQuadricSkeletalMeshReduction::ReduceSkeletalMesh(USkeletalMesh& SkeletalMe
 
 			if (!bLODModelAdded)
 			{
-				TBitArray<> SourceSectionMatched;
-				SourceSectionMatched.Init(false, DstModelBackup.Sections.Num());
-				for (int32 SectionIndex = 0; SectionIndex < ImportedModelLOD.Sections.Num(); ++SectionIndex)
-				{
-					FSkelMeshSection& Section = ImportedModelLOD.Sections[SectionIndex];
-					for (int32 SourceSectionIndex = 0; SourceSectionIndex < DstModelBackup.Sections.Num(); ++SourceSectionIndex)
-					{
-						if (SourceSectionMatched[SourceSectionIndex])
-						{
-							continue;
-						}
-						const FSkelMeshSection& SourceSection = DstModelBackup.Sections[SourceSectionIndex];
-						if (const FSkelMeshSourceSectionUserData* SourceUserData = DstModelBackup.UserSectionsData.Find(SourceSection.OriginalDataSectionIndex))
-						{
-							if (Section.MaterialIndex == SourceSection.MaterialIndex)
-							{
-								Section.OriginalDataSectionIndex = SourceSection.OriginalDataSectionIndex;
-								SourceSectionMatched[SourceSectionIndex] = true;
-								break;
-							}
-						}
-					}
-				}
-
-				ImportedModelLOD.UserSectionsData = DstModelBackup.UserSectionsData;
-				ImportedModelLOD.SyncronizeUserSectionsDataArray();
+				RestoreUserSectionsData(DstModelBackup, ImportedModelLOD);
 				//If its an existing LOD put back the buildStringID
 				ImportedModelLOD.BuildStringID = BackupLodModelBuildStringID;
 				ImportedModelLOD.RawSkeletalMeshBulkDataID = BackupRawSkeletalMeshBulkDataID;
@@ -2230,6 +2198,13 @@ void FQuadricSkeletalMeshReduction::ReduceSkeletalMesh(USkeletalMesh& SkeletalMe
 		//It also warranty that we do not change the ddc key
 		SkeletalMesh.SaveLODImportedData(LODIndex, RawMesh);
 		SkeletalMesh.SetLODImportedDataVersions(LODIndex, GeoImportVersion, SkinningImportVersion);
+		if (!bLODModelAdded)
+		{
+			//The SaveLODImportdData can change the cache RawSkeletalMeshBulkDataID, so we have to put back the backup after the save
+			FSkeletalMeshLODModel& ImportedModelLOD = SkeletalMesh.GetImportedModel()->LODModels[LODIndex];
+			ImportedModelLOD.BuildStringID = BackupLodModelBuildStringID;
+			ImportedModelLOD.RawSkeletalMeshBulkDataID = BackupRawSkeletalMeshBulkDataID;
+		}
 	}
 
 	SkeletalMesh.CalculateRequiredBones(SkeletalMeshResource.LODModels[LODIndex], SkeletalMesh.GetRefSkeleton(), &BonesToRemove);
diff --git a/Engine/Source/Developer/MeshBuilder/Private/SkeletalMeshBuilder.cpp b/Engine/Source/Developer/MeshBuilder/Private/SkeletalMeshBuilder.cpp
--- a/Engine/Source/Developer/MeshBuilder/Private/SkeletalMeshBuilder.cpp
+++ b/Engine/Source/Developer/MeshBuilder/Private/SkeletalMeshBuilder.cpp
@@ -129,12 +129,23 @@ bool FSkeletalMeshBuilder::Build(const FSkeletalMeshBuildParameters& SkeletalMes
 			Options
 		);
 
-		//Re-Apply the user section changes, the UserSectionsData is map to original section and should match the builded LODModel
-		BuildLODModel.SyncronizeUserSectionsDataArray();
-
 		// Set texture coordinate count on the new model.
 		BuildLODModel.NumTexCoords = NumTextCoord;
 
+		//Cache the vertex/triangle count in the InlineReductionCacheData so we can know if the LODModel need reduction or not.
+		TArray<FInlineReductionCacheData>& InlineReductionCacheDatas = SkeletalMesh->GetImportedModel()->InlineReductionCacheDatas;
+		if (!InlineReductionCacheDatas.IsValidIndex(LODIndex))
+		{
+			InlineReductionCacheDatas.AddDefaulted((LODIndex + 1) - InlineReductionCacheDatas.Num());
+		}
+		if (ensure(InlineReductionCacheDatas.IsValidIndex(LODIndex)))
+		{
+			InlineReductionCacheDatas[LODIndex].SetCacheGeometryInfo(BuildLODModel);
+		}
+
+		//Re-Apply the user section changes, the UserSectionsData is map to original section and should match the builded LODModel
+		BuildLODModel.SyncronizeUserSectionsDataArray();
+
 		//Re-apply the morph target
 		SlowTask.EnterProgressFrame(1.0f, NSLOCTEXT("SkeltalMeshBuilder", "RebuildMorphTarget", "Rebuilding morph targets..."));
 		if (SkeletalMeshImportData.MorphTargetNames.Num() > 0)
@@ -159,33 +170,8 @@ bool FSkeletalMeshBuilder::Build(const FSkeletalMeshBuildParameters& SkeletalMes
 		{
 			SlowTask.EnterProgressFrame(1.0f, NSLOCTEXT("SkeltalMeshBuilder", "RegenerateLOD", "Regenerate LOD..."));
 			//Update the original reduction data since we just build a new LODModel.
-			if (LODInfo->ReductionSettings.BaseLOD == LODIndex && SkeletalMesh->GetImportedModel()->OriginalReductionSourceMeshData.IsValidIndex(LODIndex))
+			if (LODInfo->ReductionSettings.BaseLOD == LODIndex && SkeletalMesh->IsLODImportedDataBuildAvailable(LODIndex))
 			{
-				// Inplace reduction is not currently thread-safe
-				check(IsInGameThread());
-
-				//Make the copy of the data only once until the ImportedModel change (re-imported)
-				SkeletalMesh->GetImportedModel()->OriginalReductionSourceMeshData[LODIndex]->EmptyBulkData();
-				TMap<FString, TArray<FMorphTargetDelta>> BaseLODMorphTargetData;
-				BaseLODMorphTargetData.Empty(SkeletalMesh->GetMorphTargets().Num());
-				for (UMorphTarget *MorphTarget : SkeletalMesh->GetMorphTargets())
-				{
-					if (!MorphTarget->HasDataForLOD(LODIndex))
-					{
-						continue;
-					}
-					TArray<FMorphTargetDelta>& MorphDeltasArray = BaseLODMorphTargetData.FindOrAdd(MorphTarget->GetFullName());
-					//Iterate each original morph target source index to fill the NewMorphTargetDeltas array with the TargetMatchData.
-					int32 NumDeltas = 0;
-					const FMorphTargetDelta* BaseDeltaArray = MorphTarget->GetMorphTargetDelta(LODIndex, NumDeltas);
-					for (int32 DeltaIndex = 0; DeltaIndex < NumDeltas; DeltaIndex++)
-					{
-						MorphDeltasArray.Add(BaseDeltaArray[DeltaIndex]);
-					}
-				}
-				//Copy the original SkeletalMesh LODModel
-				SkeletalMesh->GetImportedModel()->OriginalReductionSourceMeshData[LODIndex]->SaveReductionData(BuildLODModel, BaseLODMorphTargetData, SkeletalMesh);
-
 				if (LODIndex == 0)
 				{
 					SkeletalMesh->GetLODInfo(LODIndex)->SourceImportFilename = SkeletalMesh->GetAssetImportData()->GetFirstFilename();
diff --git a/Engine/Source/Developer/MeshUtilities/Private/MeshUtilities.cpp b/Engine/Source/Developer/MeshUtilities/Private/MeshUtilities.cpp
--- a/Engine/Source/Developer/MeshUtilities/Private/MeshUtilities.cpp
+++ b/Engine/Source/Developer/MeshUtilities/Private/MeshUtilities.cpp
@@ -5829,7 +5829,6 @@ void FMeshUtilities::CreateImportDataFromLODModel(USkeletalMesh* SkeletalMesh) c
 {
 	for (int32 LodIndex = 0; LodIndex < SkeletalMesh->GetLODNum(); LodIndex++)
 	{
-		const FSkeletalMeshLODModel* LODModel = &(SkeletalMesh->GetImportedModel()->LODModels[LodIndex]);
 		const FSkeletalMeshLODInfo* ThisLODInfo = SkeletalMesh->GetLODInfo(LodIndex);
 
 		check(ThisLODInfo);
@@ -5839,51 +5838,49 @@ void FMeshUtilities::CreateImportDataFromLODModel(USkeletalMesh* SkeletalMesh) c
 		{
 			continue;
 		}
-		const bool bReductionActive = SkeletalMesh->IsReductionActive(LodIndex);
-		const bool bInlineReduction = bReductionActive && (ThisLODInfo->ReductionSettings.BaseLOD == LodIndex);
-		if (bReductionActive && !bInlineReduction)
+		const bool bInlineReduction = SkeletalMesh->IsReductionActive(LodIndex) && (ThisLODInfo->ReductionSettings.BaseLOD == LodIndex);
+		if (!bInlineReduction)
 		{
 			//Generated LOD (not inline) do not need imported data
 			continue;
 		}
 
-		TMap<FString, TArray<FMorphTargetDelta>> LODMorphTargetData;
-		FSkeletalMeshLODModel TempLODModel;
-		if (bInlineReduction)
+		//Find the reduction data to restore the imported LODModel by using the deprecated OriginalReductionSourceMeshData_DEPRECATED data
+		if (SkeletalMesh->GetImportedModel()->OriginalReductionSourceMeshData_DEPRECATED.IsValidIndex(LodIndex))
 		{
-			//Find the reduction data to restore the imported LODModel
-			if (!SkeletalMesh->GetImportedModel()->OriginalReductionSourceMeshData.IsValidIndex(LodIndex))
-			{
-				//We should never end up here since the inline reduction feature add this data to assets
-				const bool bSkeletalMeshInlineReductionHasOriginalData = false;
-				UE_ASSET_LOG(LogSkeletalMesh, Warning, SkeletalMesh, TEXT("LOD %d do not have original reduction data but use inline reduction"), LodIndex);
-				checkSlow(bSkeletalMeshInlineReductionHasOriginalData);
-				continue;
-			}
+			FSkeletalMeshLODModel& ToUpdateLODModel = SkeletalMesh->GetImportedModel()->LODModels[LodIndex];
+			TMap<FString, TArray<FMorphTargetDelta>> LODMorphTargetData;
+			
+			//Old inline reduced assets do not have the original imported skin weight data, it must be re-import if we want to convert it to import data
 			//Swap the LODModel pointer to the one we store when reducing, so we add the true import data
-			SkeletalMesh->GetImportedModel()->OriginalReductionSourceMeshData[LodIndex]->LoadReductionData(TempLODModel, LODMorphTargetData, SkeletalMesh);
-		}
-		else
-		{
-			LODMorphTargetData.Empty(SkeletalMesh->GetMorphTargets().Num());
-			for (UMorphTarget* MorphTarget : SkeletalMesh->GetMorphTargets())
+			SkeletalMesh->GetImportedModel()->OriginalReductionSourceMeshData_DEPRECATED[LodIndex]->LoadReductionData(ToUpdateLODModel, LODMorphTargetData, SkeletalMesh);
+
+			TArray<UMorphTarget*> OriginalMorphTargets;
+			OriginalMorphTargets.Reserve(LODMorphTargetData.Num());
+			for (const TPair<FString, TArray<FMorphTargetDelta>>& MorphTargetChannel : LODMorphTargetData)
 			{
-				if (!MorphTarget->HasDataForLOD(LodIndex))
+				UMorphTarget* ExistingMorphTarget = nullptr;
+				for (UMorphTarget* MorphTarget : SkeletalMesh->GetMorphTargets())
 				{
-					continue;
+					if (MorphTarget->GetName() == MorphTargetChannel.Key)
+					{
+						ExistingMorphTarget = MorphTarget;
+					}
 				}
-				TArray<FMorphTargetDelta>& MorphDeltasArray = LODMorphTargetData.FindOrAdd(MorphTarget->GetName());
-				const FMorphTargetLODModel& BaseMorphModel = MorphTarget->GetMorphLODModels()[LodIndex];
-				//Iterate each original morph target source index to fill the NewMorphTargetDeltas array with the TargetMatchData.
-				for (const FMorphTargetDelta& MorphDelta : BaseMorphModel.Vertices)
+				if (!ExistingMorphTarget)
 				{
-					MorphDeltasArray.Add(MorphDelta);
+					FName MorphTargetName = *(MorphTargetChannel.Key);
+					ExistingMorphTarget = NewObject<UMorphTarget>(SkeletalMesh, MorphTargetName);
+					check(ExistingMorphTarget);
 				}
+
+				const TArray<FMorphTargetDelta>& MorphTargetDeltas = MorphTargetChannel.Value;
+				ExistingMorphTarget->PopulateDeltas(MorphTargetDeltas, LodIndex, ToUpdateLODModel.Sections, ThisLODInfo->BuildSettings.bRecomputeNormals, false, ThisLODInfo->BuildSettings.MorphThresholdPosition);
 			}
-			FSkeletalMeshLODModel::CopyStructure(&TempLODModel, LODModel);
 		}
 
-		//Now use the TempLODModel and the LODMorphTargetData to build an importData
+		FSkeletalMeshLODModel* LODModel = &(SkeletalMesh->GetImportedModel()->LODModels[LodIndex]);
+
 		FSkeletalMeshImportData ImportData;
 		ImportData.bDiffPose = false;
 		ImportData.bHasNormals = true;
@@ -5891,21 +5888,52 @@ void FMeshUtilities::CreateImportDataFromLODModel(USkeletalMesh* SkeletalMesh) c
 		ImportData.bUseT0AsRefPose = false;
 		ImportData.bHasVertexColors = SkeletalMesh->GetHasVertexColors();
 
-		const TArray<FSkeletalMaterial>& SKMaterials = SkeletalMesh->GetMaterials();
+		TArray<FSkeletalMaterial>& SKMaterials = SkeletalMesh->GetMaterials();
 		ImportData.Materials.Reserve(SkeletalMesh->GetMaterials().Num());
+		TArray<FName> UniqueMaterialNames;
 		for (int32 MaterialIndex = 0; MaterialIndex < SKMaterials.Num(); ++MaterialIndex)
 		{
-			SkeletalMeshImportData::FMaterial& Material = ImportData.Materials.AddDefaulted_GetRef();
+			SkeletalMeshImportData::FMaterial& Material = ImportData.Materials.AddZeroed_GetRef();
 			Material.Material = SKMaterials[MaterialIndex].MaterialInterface;
+			//Make sure skeletal mesh material has the slot name set properly
+			if (SKMaterials[MaterialIndex].ImportedMaterialSlotName == NAME_None)
+			{
+				if (SKMaterials[MaterialIndex].MaterialInterface)
+				{
+					SKMaterials[MaterialIndex].ImportedMaterialSlotName = SKMaterials[MaterialIndex].MaterialInterface->GetFName();
+				}
+				else if (SKMaterials[MaterialIndex].MaterialSlotName != NAME_None)
+				{
+					SKMaterials[MaterialIndex].ImportedMaterialSlotName = SKMaterials[MaterialIndex].MaterialSlotName;
+				}
+				else
+				{
+					//We cannot find any valid name
+					UE_ASSET_LOG(LogSkeletalMesh, Display, SkeletalMesh, TEXT("FMeshUtilities::CreateImportDataFromLODModel: Cannot assign valid slotname to material slot"));
+					SKMaterials[MaterialIndex].ImportedMaterialSlotName = NAME_None;
+				}
+				//each material slot must have a unique imported material name
+				while (UniqueMaterialNames.Contains(SKMaterials[MaterialIndex].ImportedMaterialSlotName))
+				{
+					const int32 NameUniqueNumber = SKMaterials[MaterialIndex].ImportedMaterialSlotName.GetNumber() + 1;
+					SKMaterials[MaterialIndex].ImportedMaterialSlotName.SetNumber(NameUniqueNumber);
+				}
+			}
+			if (SKMaterials[MaterialIndex].MaterialSlotName == NAME_None)
+			{
+				SKMaterials[MaterialIndex].MaterialSlotName = SKMaterials[MaterialIndex].ImportedMaterialSlotName;
+			}
+			UniqueMaterialNames.Add(SKMaterials[MaterialIndex].ImportedMaterialSlotName);
+
 			Material.MaterialImportName = SKMaterials[MaterialIndex].ImportedMaterialSlotName.ToString();
 		}
 
-		ImportData.NumTexCoords = FMath::Min(TempLODModel.NumTexCoords, static_cast<uint32>(MAX_TEXCOORDS));
+		ImportData.NumTexCoords = FMath::Min(LODModel->NumTexCoords, static_cast<uint32>(MAX_TEXCOORDS));
 
 		const FReferenceSkeleton& ReferenceSkeleton = SkeletalMesh->GetRefSkeleton();
 		for (int32 BoneIndex = 0; BoneIndex < ReferenceSkeleton.GetRawBoneNum(); ++BoneIndex)
 		{
-			SkeletalMeshImportData::FBone& RefBone = ImportData.RefBonesBinary.AddDefaulted_GetRef();
+			SkeletalMeshImportData::FBone& RefBone = ImportData.RefBonesBinary.AddZeroed_GetRef();
 			const FMeshBoneInfo& MeshBoneInfo = ReferenceSkeleton.GetRawRefBoneInfo()[BoneIndex];
 			const FTransform& MeshBonePose = ReferenceSkeleton.GetRawRefBonePose()[BoneIndex];
 			RefBone.Name = MeshBoneInfo.ExportName;
@@ -5920,17 +5948,17 @@ void FMeshUtilities::CreateImportDataFromLODModel(USkeletalMesh* SkeletalMesh) c
 		}
 
 		TArray<FSoftSkinVertex> Vertices;
-		TempLODModel.GetVertices(Vertices);
+		LODModel->GetVertices(Vertices);
 		const int32 VertexCount = Vertices.Num();
-		const int32 TriangleCount = TempLODModel.IndexBuffer.Num() / 3;
+		const int32 TriangleCount = LODModel->IndexBuffer.Num() / 3;
 		ImportData.Points.Reserve(VertexCount);
 		ImportData.PointToRawMap.Reserve(VertexCount);
 		ImportData.Wedges.Reserve(VertexCount);
 		ImportData.Faces.Reserve(TriangleCount);
 
-		for (int32 SectionIndex = 0; SectionIndex < TempLODModel.Sections.Num(); ++SectionIndex)
+		for (int32 SectionIndex = 0; SectionIndex < LODModel->Sections.Num(); ++SectionIndex)
 		{
-			const FSkelMeshSection& Section = TempLODModel.Sections[SectionIndex];
+			const FSkelMeshSection& Section = LODModel->Sections[SectionIndex];
 			int32 VerticeIndexStart = Section.BaseVertexIndex;
 			int32 VerticeIndexEnd = VerticeIndexStart + Section.GetNumVertices();
 			int32 MaterialIndex = Section.MaterialIndex;
@@ -5948,7 +5976,7 @@ void FMeshUtilities::CreateImportDataFromLODModel(USkeletalMesh* SkeletalMesh) c
 					{
 						break;
 					}
-					SkeletalMeshImportData::FRawBoneInfluence& Influence = ImportData.Influences.AddDefaulted_GetRef();
+					SkeletalMeshImportData::FRawBoneInfluence& Influence = ImportData.Influences.AddZeroed_GetRef();
 					Influence.VertexIndex = PointIndex;
 					Influence.BoneIndex = Section.BoneMap[Vertex.InfluenceBones[InfluenceIndex]];
 					Influence.Weight = static_cast<float>(Vertex.InfluenceWeights[InfluenceIndex]) / 255.0f;
@@ -5958,7 +5986,7 @@ void FMeshUtilities::CreateImportDataFromLODModel(USkeletalMesh* SkeletalMesh) c
 
 		for (int32 TriangleIndex = 0; TriangleIndex < TriangleCount; ++TriangleIndex)
 		{
-			SkeletalMeshImportData::FTriangle& Face = ImportData.Faces.AddDefaulted_GetRef();
+			SkeletalMeshImportData::FTriangle& Face = ImportData.Faces.AddZeroed_GetRef();
 			int32 IndexBufferIndex = TriangleIndex * 3;
 			int32 SectionIndex;
 			int32 VerticeIndex[3];
@@ -5966,25 +5994,25 @@ void FMeshUtilities::CreateImportDataFromLODModel(USkeletalMesh* SkeletalMesh) c
 			for (int32 Corner = 0; Corner < 3; ++Corner)
 			{
 				const int32 WedgeIndex = ImportData.Wedges.Num();
-				VerticeIndex[Corner] = TempLODModel.IndexBuffer[IndexBufferIndex + Corner];
+				VerticeIndex[Corner] = LODModel->IndexBuffer[IndexBufferIndex + Corner];
 				int32 NextEdgeCornerIndex = (Corner + 1) % 3;
-				TempLODModel.GetSectionFromVertexIndex(VerticeIndex[Corner], SectionIndex, SectionVertexIndex[Corner]);
+				LODModel->GetSectionFromVertexIndex(VerticeIndex[Corner], SectionIndex, SectionVertexIndex[Corner]);
 				const FSoftSkinVertex& Vertex = Vertices[VerticeIndex[Corner]];
-				SkeletalMeshImportData::FVertex& Wedge = ImportData.Wedges.AddDefaulted_GetRef();
+				SkeletalMeshImportData::FVertex& Wedge = ImportData.Wedges.AddZeroed_GetRef();
 				Wedge.Color = Vertex.Color;
 				for (int32 UVIndex = 0; UVIndex < static_cast<int32>(ImportData.NumTexCoords); ++UVIndex)
 				{
 					Wedge.UVs[UVIndex] = Vertex.UVs[UVIndex];
 				}
 				Wedge.VertexIndex = VerticeIndex[Corner];
-				Wedge.MatIndex = TempLODModel.Sections[SectionIndex].MaterialIndex;
+				Wedge.MatIndex = LODModel->Sections[SectionIndex].MaterialIndex;
 
 				Face.WedgeIndex[Corner] = WedgeIndex;
 				Face.TangentX[Corner] = Vertex.TangentX;
 				Face.TangentY[Corner] = Vertex.TangentY;
 				Face.TangentZ[Corner] = Vertex.TangentZ;
 			}
-			Face.MatIndex = TempLODModel.Sections[SectionIndex].MaterialIndex;
+			Face.MatIndex = LODModel->Sections[SectionIndex].MaterialIndex;
 			//Faceted by default
 			Face.SmoothingGroups = 0x0;
 		}
@@ -5992,29 +6020,35 @@ void FMeshUtilities::CreateImportDataFromLODModel(USkeletalMesh* SkeletalMesh) c
 		//Recreate smooth group
 		ImportData.ComputeSmoothGroupFromNormals();
 
-		ImportData.MorphTargetNames.Reserve(LODMorphTargetData.Num());
-		ImportData.MorphTargetModifiedPoints.Reserve(LODMorphTargetData.Num());
-		ImportData.MorphTargets.Reserve(LODMorphTargetData.Num());
-		for (const TPair<FString, TArray<FMorphTargetDelta>>& MorphTargetChannel : LODMorphTargetData)
+		const int32 SrcMorphTargetNum = SkeletalMesh->GetMorphTargets().Num();
+		ImportData.MorphTargetNames.Reserve(SrcMorphTargetNum);
+		ImportData.MorphTargetModifiedPoints.Reserve(SrcMorphTargetNum);
+		ImportData.MorphTargets.Reserve(SrcMorphTargetNum);
+		for (UMorphTarget* MorphTarget : SkeletalMesh->GetMorphTargets())
 		{
-			ImportData.MorphTargetNames.Add(MorphTargetChannel.Key);
-			const TArray<FMorphTargetDelta>& MorphTargetDeltas = MorphTargetChannel.Value;
+			if (!MorphTarget->HasDataForLOD(LodIndex))
+			{
+				continue;
+			}
+			ImportData.MorphTargetNames.Add(MorphTarget->GetName());
+			const FMorphTargetLODModel& BaseMorphModel = MorphTarget->GetMorphLODModels()[LodIndex];
 			FSkeletalMeshImportData& MorphImportShape = ImportData.MorphTargets.AddDefaulted_GetRef();
 			TSet<uint32>& ModifiedPoints = ImportData.MorphTargetModifiedPoints.AddDefaulted_GetRef();
-			for (const FMorphTargetDelta& MorphTargetDelta : MorphTargetDeltas)
+			//Iterate each original morph target source index to fill the NewMorphTargetDeltas array with the TargetMatchData.
+			for (const FMorphTargetDelta& MorphTargetDelta : BaseMorphModel.Vertices)
 			{
 				ModifiedPoints.Add(MorphTargetDelta.SourceIdx);
 				MorphImportShape.Points.Add(MorphTargetDelta.PositionDelta + Vertices[MorphTargetDelta.SourceIdx].Position);
 			}
 		}
 
-		if (TempLODModel.SkinWeightProfiles.Num() > 0)
+		if (LODModel->SkinWeightProfiles.Num() > 0)
 		{
 			FSkeletalMeshImportData ReferenceCopy;
 			ReferenceCopy = ImportData;
 			ReferenceCopy.KeepAlternateSkinningBuildDataOnly();
 			ReferenceCopy.Influences.Empty();
-			for (TPair<FName, FImportedSkinWeightProfileData>& SkinWeightProfile : TempLODModel.SkinWeightProfiles)
+			for (TPair<FName, FImportedSkinWeightProfileData>& SkinWeightProfile : LODModel->SkinWeightProfiles)
 			{
 				FString ProfileName = SkinWeightProfile.Key.ToString();
 				FImportedSkinWeightProfileData& ImportedSkinWeightProfileData = SkinWeightProfile.Value;
@@ -6032,8 +6066,8 @@ void FMeshUtilities::CreateImportDataFromLODModel(USkeletalMesh* SkeletalMesh) c
 				{
 					int32 OutSectionIndex = INDEX_NONE;
 					int32 OutSectionVertexIndex = INDEX_NONE;
-					TempLODModel.GetSectionFromVertexIndex(VerticeIndex, OutSectionIndex, OutSectionVertexIndex);
-					const FSkelMeshSection& Section = TempLODModel.Sections[OutSectionIndex];
+					LODModel->GetSectionFromVertexIndex(VerticeIndex, OutSectionIndex, OutSectionVertexIndex);
+					const FSkelMeshSection& Section = LODModel->Sections[OutSectionIndex];
 					const FRawSkinWeight& SkinWeight = ImportedSkinWeightProfileData.SkinWeights[VerticeIndex];
 					for (int32 InfluenceIndex = 0; InfluenceIndex < MAX_TOTAL_INFLUENCES; ++InfluenceIndex)
 					{
@@ -6042,11 +6076,11 @@ void FMeshUtilities::CreateImportDataFromLODModel(USkeletalMesh* SkeletalMesh) c
 						{
 							break;
 						}
-						SkeletalMeshImportData::FRawBoneInfluence& Influence = AlternateInfluence.Influences.AddDefaulted_GetRef();
+						SkeletalMeshImportData::FRawBoneInfluence& Influence = AlternateInfluence.Influences.AddZeroed_GetRef();
 						Influence.VertexIndex = VerticeIndex;
 						Influence.BoneIndex = Section.BoneMap[SkinWeight.InfluenceBones[InfluenceIndex]];
 						Influence.Weight = Weight;
-						SkeletalMeshImportData::FVertInfluence& SourceModelInfluence = ImportedSkinWeightProfileData.SourceModelInfluences.AddDefaulted_GetRef();
+						SkeletalMeshImportData::FVertInfluence& SourceModelInfluence = ImportedSkinWeightProfileData.SourceModelInfluences.AddZeroed_GetRef();
 						SourceModelInfluence.VertIndex = VerticeIndex;
 						SourceModelInfluence.BoneIndex = SkinWeight.InfluenceBones[InfluenceIndex];
 						SourceModelInfluence.Weight = Weight;
@@ -6058,6 +6092,7 @@ void FMeshUtilities::CreateImportDataFromLODModel(USkeletalMesh* SkeletalMesh) c
 
 		SkeletalMesh->SaveLODImportedData(LodIndex, ImportData);
 		SkeletalMesh->SetLODImportedDataVersions(LodIndex, ESkeletalMeshGeoImportVersions::LatestVersion, ESkeletalMeshSkinningImportVersions::LatestVersion);
+		UE_ASSET_LOG(LogSkeletalMesh, Display, SkeletalMesh, TEXT("FMeshUtilities::CreateImportDataFromLODModel: ImportData created for LOD %d"), LodIndex);
 	}
 }
 
diff --git a/Engine/Source/Developer/SkeletalMeshUtilitiesCommon/Private/LODUtilities.cpp b/Engine/Source/Developer/SkeletalMeshUtilitiesCommon/Private/LODUtilities.cpp
--- a/Engine/Source/Developer/SkeletalMeshUtilitiesCommon/Private/LODUtilities.cpp
+++ b/Engine/Source/Developer/SkeletalMeshUtilitiesCommon/Private/LODUtilities.cpp
@@ -374,23 +374,20 @@ void FLODUtilities::RemoveLOD(FSkeletalMeshUpdateContext& UpdateContext, int32 D
 				SkeletalMesh->SaveLODImportedData(FirstDepLODIndex, ToRemovedLODImportData);
 
 				//Manage the override original reduction source mesh data
-				if (SkelMeshModel->OriginalReductionSourceMeshData.IsValidIndex(FirstDepLODIndex) && SkelMeshModel->OriginalReductionSourceMeshData[FirstDepLODIndex])
+				if (SkelMeshModel->InlineReductionCacheDatas.IsValidIndex(FirstDepLODIndex))
 				{
 					if(SkeletalMesh->IsLODImportedDataBuildAvailable(DesiredLOD))
 					{
-						//Empty the OriginalReductionSourceMeshData it will be recreate when we will rebuild the asset
-						SkelMeshModel->OriginalReductionSourceMeshData[FirstDepLODIndex]->EmptyBulkData();
+						//The inline reduction cache data will be recache by the build
+						SkelMeshModel->InlineReductionCacheDatas[FirstDepLODIndex].SetCacheGeometryInfo(MAX_uint32, MAX_uint32);
 					}
-					//If we are modifying an old asset we have to duplicate the original reduction source mesh data from the source model
-					else if(SkelMeshModel->OriginalReductionSourceMeshData.IsValidIndex(DesiredLOD) &&
-						SkelMeshModel->OriginalReductionSourceMeshData[DesiredLOD] &&
-						!SkelMeshModel->OriginalReductionSourceMeshData[DesiredLOD]->IsEmpty())
+					else if(SkelMeshModel->InlineReductionCacheDatas.IsValidIndex(DesiredLOD))
 					{
 						//If there is no build copy the one from the DesiredLOD
-						FSkeletalMeshLODModel LODModel;
-						TMap<FString, TArray<FMorphTargetDelta>> BaseLODMorphTargetData;
-						SkelMeshModel->OriginalReductionSourceMeshData[DesiredLOD]->LoadReductionData(LODModel, BaseLODMorphTargetData, SkeletalMesh);
-						SkelMeshModel->OriginalReductionSourceMeshData[FirstDepLODIndex]->SaveReductionData(LODModel, BaseLODMorphTargetData, SkeletalMesh);
+						uint32 CacheVertexCount = 0;
+						uint32 CacheTriangleCount = 0;
+						SkelMeshModel->InlineReductionCacheDatas[DesiredLOD].GetCacheGeometryInfo(CacheVertexCount, CacheTriangleCount);
+						SkelMeshModel->InlineReductionCacheDatas[FirstDepLODIndex].SetCacheGeometryInfo(CacheVertexCount, CacheTriangleCount);
 					}
 				}
 
@@ -840,15 +837,8 @@ void ProjectTargetOnBase(const TArray<FSoftSkinVertex>& BaseVertices, const TArr
 	}
 }
 
-void CreateLODMorphTarget(USkeletalMesh* SkeletalMesh, FReductionBaseSkeletalMeshBulkData* ReductionBaseSkeletalMeshBulkData, int32 SourceLOD, int32 DestinationLOD, const TMap<UMorphTarget *, TMap<uint32, uint32>>& PerMorphTargetBaseIndexToMorphTargetDelta, const TMap<uint32, TArray<uint32>>& BaseMorphIndexToTargetIndexList, const TArray<FSoftSkinVertex>& TargetVertices, const TArray<FTargetMatch>& TargetMatchData)
+void CreateLODMorphTarget(USkeletalMesh* SkeletalMesh, const FInlineReductionDataParameter& InlineReductionDataParameter, int32 SourceLOD, int32 DestinationLOD, const TMap<UMorphTarget *, TMap<uint32, uint32>>& PerMorphTargetBaseIndexToMorphTargetDelta, const TMap<uint32, TArray<uint32>>& BaseMorphIndexToTargetIndexList, const TArray<FSoftSkinVertex>& TargetVertices, const TArray<FTargetMatch>& TargetMatchData)
 {
-	TMap<FString, TArray<FMorphTargetDelta>> BaseLODMorphTargetData;
-	if (ReductionBaseSkeletalMeshBulkData != nullptr)
-	{
-		FSkeletalMeshLODModel TempBaseLODModel;
-		ReductionBaseSkeletalMeshBulkData->LoadReductionData(TempBaseLODModel, BaseLODMorphTargetData, SkeletalMesh);
-	}
-
 	FSkeletalMeshModel* SkeletalMeshModel = SkeletalMesh->GetImportedModel();
 	const FSkeletalMeshLODModel& TargetLODModel = SkeletalMeshModel->LODModels[DestinationLOD];
 
@@ -860,9 +850,9 @@ void CreateLODMorphTarget(USkeletalMesh* SkeletalMesh, FReductionBaseSkeletalMes
 		{
 			continue;
 		}
-		bool bUseBaseMorphDelta = SourceLOD == DestinationLOD && BaseLODMorphTargetData.Contains(MorphTarget->GetFullName());
+		bool bUseBaseMorphDelta = SourceLOD == DestinationLOD && InlineReductionDataParameter.bIsDataValid && InlineReductionDataParameter.InlineOriginalSrcMorphTargetData.Contains(MorphTarget->GetFullName());
 
-		const TArray<FMorphTargetDelta> *BaseMorphDeltas = bUseBaseMorphDelta ? BaseLODMorphTargetData.Find(MorphTarget->GetFullName()) : nullptr;
+		const TArray<FMorphTargetDelta> *BaseMorphDeltas = bUseBaseMorphDelta ? InlineReductionDataParameter.InlineOriginalSrcMorphTargetData.Find(MorphTarget->GetFullName()) : nullptr;
 		if (BaseMorphDeltas == nullptr || BaseMorphDeltas->Num() <= 0)
 		{
 			bUseBaseMorphDelta = false;
@@ -996,7 +986,7 @@ void FLODUtilities::ClearGeneratedMorphTarget(USkeletalMesh* SkeletalMesh, int32
 	}
 }
 
-void FLODUtilities::ApplyMorphTargetsToLOD(USkeletalMesh* SkeletalMesh, int32 SourceLOD, int32 DestinationLOD)
+void FLODUtilities::ApplyMorphTargetsToLOD(USkeletalMesh* SkeletalMesh, int32 SourceLOD, int32 DestinationLOD, const FInlineReductionDataParameter& InlineReductionDataParameter)
 {
 	check(SkeletalMesh);
 	FSkeletalMeshModel* SkeletalMeshResource = SkeletalMesh->GetImportedModel();
@@ -1010,8 +1000,7 @@ void FLODUtilities::ApplyMorphTargetsToLOD(USkeletalMesh* SkeletalMesh, int32 So
 	}
 
 	FSkeletalMeshLODModel& SourceLODModel = SkeletalMeshResource->LODModels[SourceLOD];
-	FReductionBaseSkeletalMeshBulkData* ReductionBaseSkeletalMeshBulkData = nullptr;
-	bool bReduceBaseLOD = DestinationLOD == SourceLOD && SkeletalMeshResource->OriginalReductionSourceMeshData.IsValidIndex(SourceLOD) && !SkeletalMeshResource->OriginalReductionSourceMeshData[SourceLOD]->IsEmpty();
+	bool bReduceBaseLOD = DestinationLOD == SourceLOD && InlineReductionDataParameter.bIsDataValid;
 	if (!bReduceBaseLOD && SourceLOD == DestinationLOD)
 	{
 		//Abort remapping of morph target since the data is missing
@@ -1033,20 +1022,7 @@ void FLODUtilities::ApplyMorphTargetsToLOD(USkeletalMesh* SkeletalMesh, int32 So
 		return;
 	}
 
-	if (bReduceBaseLOD)
-	{
-		ReductionBaseSkeletalMeshBulkData = SkeletalMeshResource->OriginalReductionSourceMeshData[SourceLOD];
-	}
-
-	FSkeletalMeshLODModel TempBaseLODModel;
-	TMap<FString, TArray<FMorphTargetDelta>> TempBaseLODMorphTargetData;
-	if (bReduceBaseLOD)
-	{
-		check(ReductionBaseSkeletalMeshBulkData != nullptr);
-		ReductionBaseSkeletalMeshBulkData->LoadReductionData(TempBaseLODModel, TempBaseLODMorphTargetData, SkeletalMesh);
-	}
-
-	const FSkeletalMeshLODModel& BaseLODModel = bReduceBaseLOD ? TempBaseLODModel : SkeletalMeshResource->LODModels[SourceLOD];
+	const FSkeletalMeshLODModel& BaseLODModel = bReduceBaseLOD ? InlineReductionDataParameter.InlineOriginalSrcModel : SkeletalMeshResource->LODModels[SourceLOD];
 	const FSkeletalMeshLODInfo* BaseLODInfo = SkeletalMesh->GetLODInfo(SourceLOD);
 	const FSkeletalMeshLODModel& TargetLODModel = SkeletalMeshResource->LODModels[DestinationLOD];
 	const FSkeletalMeshLODInfo* TargetLODInfo = SkeletalMesh->GetLODInfo(DestinationLOD);
@@ -1103,7 +1079,12 @@ void FLODUtilities::ApplyMorphTargetsToLOD(USkeletalMesh* SkeletalMesh, int32 So
 		}
 	}
 	//We should have match all the target sections
-	check(!TargetSectionMatchBaseIndex.Contains(INDEX_NONE));
+	if (TargetSectionMatchBaseIndex.Contains(INDEX_NONE))
+	{
+		//This case is not fatal but need attention.
+		//Because of the chunking its possible a generated LOD end up with more sections.
+		UE_ASSET_LOG(LogLODUtilities, Display, SkeletalMesh, TEXT("FLODUtilities::ApplyMorphTargetsToLOD: The target contain more section then the source. Extra sections will not be affected by morph targets remap"));
+	}
 	TArray<FSoftSkinVertex> BaseVertices;
 	TArray<FSoftSkinVertex> TargetVertices;
 	BaseLODModel.GetVertices(BaseVertices);
@@ -1141,8 +1122,8 @@ void FLODUtilities::ApplyMorphTargetsToLOD(USkeletalMesh* SkeletalMesh, int32 So
 			continue;
 		}
 
-		bool bUseTempMorphDelta = SourceLOD == DestinationLOD && bReduceBaseLOD && TempBaseLODMorphTargetData.Contains(MorphTarget->GetFullName());
-		const TArray<FMorphTargetDelta> *TempMorphDeltas = bUseTempMorphDelta ? TempBaseLODMorphTargetData.Find(MorphTarget->GetFullName()) : nullptr;
+		bool bUseTempMorphDelta = SourceLOD == DestinationLOD && bReduceBaseLOD && InlineReductionDataParameter.InlineOriginalSrcMorphTargetData.Contains(MorphTarget->GetFullName());
+		const TArray<FMorphTargetDelta> *TempMorphDeltas = bUseTempMorphDelta ? InlineReductionDataParameter.InlineOriginalSrcMorphTargetData.Find(MorphTarget->GetFullName()) : nullptr;
 		if (TempMorphDeltas == nullptr || TempMorphDeltas->Num() <= 0)
 		{
 			bUseTempMorphDelta = false;
@@ -1173,7 +1154,7 @@ void FLODUtilities::ApplyMorphTargetsToLOD(USkeletalMesh* SkeletalMesh, int32 So
 		}
 	}
 	//Create the target morph target
-	CreateLODMorphTarget(SkeletalMesh, ReductionBaseSkeletalMeshBulkData, SourceLOD, DestinationLOD, PerMorphTargetBaseIndexToMorphTargetDelta, BaseMorphIndexToTargetIndexList, TargetVertices, TargetMatchData);
+	CreateLODMorphTarget(SkeletalMesh, InlineReductionDataParameter, SourceLOD, DestinationLOD, PerMorphTargetBaseIndexToMorphTargetDelta, BaseMorphIndexToTargetIndexList, TargetVertices, TargetMatchData);
 }
 
 void FLODUtilities::SimplifySkeletalMeshLOD( USkeletalMesh* SkeletalMesh, int32 DesiredLOD, const ITargetPlatform* TargetPlatform, bool bRestoreClothing /*= false*/, FThreadSafeBool* OutNeedsPackageDirtied/*= nullptr*/)
@@ -1194,7 +1175,7 @@ void FLODUtilities::SimplifySkeletalMeshLOD( USkeletalMesh* SkeletalMesh, int32
 	if (DesiredLOD == 0
 		&& SkeletalMesh->GetLODInfo(DesiredLOD) != nullptr
 		&& SkeletalMesh->GetLODInfo(DesiredLOD)->bHasBeenSimplified
-		&& (!SkeletalMesh->GetImportedModel()->OriginalReductionSourceMeshData.IsValidIndex(0) || SkeletalMesh->GetImportedModel()->OriginalReductionSourceMeshData[0]->IsEmpty()))
+		&& !SkeletalMesh->GetImportedModel()->InlineReductionCacheDatas.IsValidIndex(0))
 	{
 		//The base LOD was reduce and there is no valid data, we cannot regenerate this lod it must be re-import before
 		FFormatNamedArguments Args;
@@ -1231,6 +1212,8 @@ void FLODUtilities::SimplifySkeletalMeshLOD( USkeletalMesh* SkeletalMesh, int32
 		FLODUtilities::UnbindClothingAndBackup(SkeletalMesh, ClothingBindings, DesiredLOD);
 	}
 
+	FInlineReductionDataParameter InlineReductionDataParameter;
+
 	if (SkeletalMesh->GetLODInfo(DesiredLOD) != nullptr)
 	{
 		FSkeletalMeshModel* SkeletalMeshResource = SkeletalMesh->GetImportedModel();
@@ -1239,56 +1222,58 @@ void FLODUtilities::SimplifySkeletalMeshLOD( USkeletalMesh* SkeletalMesh, int32
 		//We must save the original reduction data, special case when we reduce inline we save even if its already simplified
 		if (SkeletalMeshResource->LODModels.IsValidIndex(DesiredLOD) && (!SkeletalMesh->GetLODInfo(DesiredLOD)->bHasBeenSimplified || DesiredLOD == Settings.BaseLOD))
 		{
-			// Caller should not call this in multithread if the LOD was never simplified or if its doing inline reduction
-			check(IsInGameThread());
 
 			FSkeletalMeshLODModel& SrcModel = SkeletalMeshResource->LODModels[DesiredLOD];
-			while (DesiredLOD >= SkeletalMeshResource->OriginalReductionSourceMeshData.Num())
+			if (!SkeletalMeshResource->InlineReductionCacheDatas.IsValidIndex(DesiredLOD))
 			{
-				FReductionBaseSkeletalMeshBulkData *EmptyReductionData = new FReductionBaseSkeletalMeshBulkData();
-				SkeletalMeshResource->OriginalReductionSourceMeshData.Add(EmptyReductionData);
+				//We should not do that in a worker thread, the serialization of the SkeletalMeshResource is suppose to allocate the correct number of inline data caches
+				//If the user add LOD in person editor, the simplification will be call in the game thread, see FLODUtilities::RegenerateLOD
+				if (!ensure(IsInGameThread()))
+				{
+					UE_ASSET_LOG(LogLODUtilities, Error, SkeletalMesh, TEXT("FLODUtilities::SimplifySkeletalMeshLOD: InlineReductionCacheDatas was not added in the game thread."));
+				}
+				SkeletalMeshResource->InlineReductionCacheDatas.AddDefaulted((DesiredLOD + 1) - SkeletalMeshResource->InlineReductionCacheDatas.Num());
 			}
-			check(SkeletalMeshResource->OriginalReductionSourceMeshData.IsValidIndex(DesiredLOD));
-			//Make the copy of the data only once until the ImportedModel change (re-imported)
-			if (SkeletalMeshResource->OriginalReductionSourceMeshData[DesiredLOD]->IsEmpty())
+			check(SkeletalMeshResource->InlineReductionCacheDatas.IsValidIndex(DesiredLOD));
+			SkeletalMeshResource->InlineReductionCacheDatas[DesiredLOD].SetCacheGeometryInfo(SrcModel);
+
+			InlineReductionDataParameter.InlineOriginalSrcMorphTargetData.Empty(SkeletalMesh->GetMorphTargets().Num());
+			for (UMorphTarget* MorphTarget : SkeletalMesh->GetMorphTargets())
 			{
-				TMap<FString, TArray<FMorphTargetDelta>> BaseLODMorphTargetData;
-				BaseLODMorphTargetData.Empty(SkeletalMesh->GetMorphTargets().Num());
-				for (UMorphTarget *MorphTarget : SkeletalMesh->GetMorphTargets())
+				if (!MorphTarget->HasDataForLOD(DesiredLOD))
 				{
-					if (!MorphTarget->HasDataForLOD(DesiredLOD))
-					{
-						continue;
-					}
-					TArray<FMorphTargetDelta>& MorphDeltasArray = BaseLODMorphTargetData.FindOrAdd(MorphTarget->GetFullName());
-					//Iterate each original morph target source index to fill the NewMorphTargetDeltas array with the TargetMatchData.
-					int32 NumDeltas = 0;
-					const FMorphTargetDelta* BaseDeltaArray = MorphTarget->GetMorphTargetDelta(DesiredLOD, NumDeltas);
-					for (int32 DeltaIndex = 0; DeltaIndex < NumDeltas; DeltaIndex++)
-					{
-						MorphDeltasArray.Add(BaseDeltaArray[DeltaIndex]);
-					}
+					continue;
 				}
-				
-				// Copy the original SkeletalMesh LODModel
-				// Unbind clothing before saving the original data, we must not restore clothing to do inline reduction
+				TArray<FMorphTargetDelta>& MorphDeltasArray = InlineReductionDataParameter.InlineOriginalSrcMorphTargetData.FindOrAdd(MorphTarget->GetFullName());
+				const FMorphTargetLODModel& BaseMorphModel = MorphTarget->GetMorphLODModels()[DesiredLOD];
+				//Iterate each original morph target source index to fill the NewMorphTargetDeltas array with the TargetMatchData.
+				int32 NumDeltas = 0;
+				const FMorphTargetDelta* BaseDeltaArray = MorphTarget->GetMorphTargetDelta(DesiredLOD, NumDeltas);
+				for (int32 DeltaIndex = 0; DeltaIndex < NumDeltas; DeltaIndex++)
 				{
-					TArray<ClothingAssetUtils::FClothingAssetMeshBinding> TemporaryRemoveClothingBindings;
-					FLODUtilities::UnbindClothingAndBackup(SkeletalMesh, TemporaryRemoveClothingBindings, DesiredLOD);
+					MorphDeltasArray.Add(BaseDeltaArray[DeltaIndex]);
+				}
+			}
 
-					SkeletalMeshResource->OriginalReductionSourceMeshData[DesiredLOD]->SaveReductionData(SrcModel, BaseLODMorphTargetData, SkeletalMesh);
+			// Copy the original SkeletalMesh LODModel
+			// Unbind clothing before saving the original data, we must not restore clothing to do inline reduction
+			{
+				TArray<ClothingAssetUtils::FClothingAssetMeshBinding> TemporaryRemoveClothingBindings;
+				FLODUtilities::UnbindClothingAndBackup(SkeletalMesh, TemporaryRemoveClothingBindings, DesiredLOD);
 
-					if (TemporaryRemoveClothingBindings.Num() > 0)
-					{
-						FLODUtilities::RestoreClothingFromBackup(SkeletalMesh, TemporaryRemoveClothingBindings, DesiredLOD);
-					}
-				}
+				FSkeletalMeshLODModel::CopyStructure(&InlineReductionDataParameter.InlineOriginalSrcModel, &SrcModel);
 
-				if (DesiredLOD == 0)
+				if (TemporaryRemoveClothingBindings.Num() > 0)
 				{
-					SkeletalMesh->GetLODInfo(DesiredLOD)->SourceImportFilename = SkeletalMesh->GetAssetImportData()->GetFirstFilename();
+					FLODUtilities::RestoreClothingFromBackup(SkeletalMesh, TemporaryRemoveClothingBindings, DesiredLOD);
 				}
 			}
+			InlineReductionDataParameter.bIsDataValid = true;
+
+			if (DesiredLOD == 0)
+			{
+				SkeletalMesh->GetLODInfo(DesiredLOD)->SourceImportFilename = SkeletalMesh->GetAssetImportData()->GetFirstFilename();
+			}
 		}
 	}
 
@@ -1296,21 +1281,20 @@ void FLODUtilities::SimplifySkeletalMeshLOD( USkeletalMesh* SkeletalMesh, int32
 	{
 		check(SkeletalMesh->GetLODNum() >= 1);
 
-		auto ApplyMorphTargetOption = [&SkeletalMesh, &DesiredLOD]()
+		//Manage morph target after the reduction. either apply to the reduce LOD or clear them all
 		{
 			FSkeletalMeshOptimizationSettings& ReductionSettings = SkeletalMesh->GetLODInfo(DesiredLOD)->ReductionSettings;
 			//Apply morph to the new LOD. Force it if we reduce the base LOD, base LOD must apply the morph target
 			if (ReductionSettings.bRemapMorphTargets)
 			{
-				ApplyMorphTargetsToLOD(SkeletalMesh, ReductionSettings.BaseLOD, DesiredLOD);
+				ApplyMorphTargetsToLOD(SkeletalMesh, ReductionSettings.BaseLOD, DesiredLOD, InlineReductionDataParameter);
 			}
 			else
 			{
 				ClearGeneratedMorphTarget(SkeletalMesh, DesiredLOD);
 			}
-		};
+		}
 
-		ApplyMorphTargetOption();
 		if (IsInGameThread())
 		{
 			SkeletalMesh->MarkPackageDirty();
@@ -1367,57 +1351,12 @@ void FLODUtilities::SimplifySkeletalMeshLOD(FSkeletalMeshUpdateContext& UpdateCo
 	}
 }
 
-bool FLODUtilities::RestoreSkeletalMeshLODImportedData(USkeletalMesh* SkeletalMesh, int32 LodIndex)
+bool FLODUtilities::RestoreSkeletalMeshLODImportedData_DEPRECATED(USkeletalMesh* SkeletalMesh, int32 LodIndex)
 {
-	if (!SkeletalMesh->GetImportedModel()->OriginalReductionSourceMeshData.IsValidIndex(LodIndex) || SkeletalMesh->GetImportedModel()->OriginalReductionSourceMeshData[LodIndex]->IsEmpty())
-	{
-		//There is nothing to restore
-		return false;
-	}
-
-	FScopedSkeletalMeshPostEditChange ScopedPostEditChange(SkeletalMesh);
-
-	// Unbind LodIndex existing clothing assets before restoring the LOD
-	TArray<ClothingAssetUtils::FClothingAssetMeshBinding> ClothingBindings;
-	FLODUtilities::UnbindClothingAndBackup(SkeletalMesh, ClothingBindings);
-
-	FSkeletalMeshLODModel ImportedBaseLODModel;
-	TMap<FString, TArray<FMorphTargetDelta>> ImportedBaseLODMorphTargetData;
-	SkeletalMesh->GetImportedModel()->OriginalReductionSourceMeshData[LodIndex]->LoadReductionData(ImportedBaseLODModel, ImportedBaseLODMorphTargetData, SkeletalMesh);
-	{
-		if (!SkeletalMesh->IsLODImportedDataBuildAvailable(LodIndex))
-		{
-			ImportedBaseLODModel.UpdateChunkedSectionInfo(SkeletalMesh->GetName());
-		}
-		//When we restore a LOD we destroy the LODMaterialMap (user manual section material slot assignation)
-		FSkeletalMeshLODInfo* LODInfo = SkeletalMesh->GetLODInfo(LodIndex);
-		LODInfo->LODMaterialMap.Empty();
-		LODInfo->bHasBeenSimplified = false;
-		//Copy the SkeletalMeshLODModel
-		FSkeletalMeshLODModel::CopyStructure(&(SkeletalMesh->GetImportedModel()->LODModels[LodIndex]), &ImportedBaseLODModel);
-		//Copy the morph target deltas
-		bool bInitMorphTargetData = false;
-		for (UMorphTarget *MorphTarget : SkeletalMesh->GetMorphTargets())
-		{
-			if (!ImportedBaseLODMorphTargetData.Contains(MorphTarget->GetFullName()))
-			{
-				continue;
-			}
-			TArray<FMorphTargetDelta>& ImportedDeltas = ImportedBaseLODMorphTargetData[MorphTarget->GetFullName()];
-
-			MorphTarget->PopulateDeltas(ImportedDeltas, LodIndex, SkeletalMesh->GetImportedModel()->LODModels[LodIndex].Sections, false, false);
-			bInitMorphTargetData |= SkeletalMesh->RegisterMorphTarget(MorphTarget, false);
-		}
-		SkeletalMesh->InitMorphTargetsAndRebuildRenderData();
-		
-		//Empty the bulkdata since we restore it
-		SkeletalMesh->GetImportedModel()->OriginalReductionSourceMeshData[LodIndex]->EmptyBulkData();
-
-		//Put back the clothing for the restore LOD
-		FLODUtilities::RestoreClothingFromBackup(SkeletalMesh, ClothingBindings);
-	}
-
-	return true;
+	const bool bThisFunctionIsDeprecated = true;
+	ensure(!bThisFunctionIsDeprecated);
+	UE_ASSET_LOG(LogLODUtilities, Error, SkeletalMesh, TEXT("FLODUtilities::RestoreSkeletalMeshLODImportedData_DEPRECATED: This function is deprecated."));
+	return false;
 }
 
 void FLODUtilities::RefreshLODChange(const USkeletalMesh* SkeletalMesh)
@@ -2370,32 +2309,28 @@ void FLODUtilities::RegenerateDependentLODs(USkeletalMesh* SkeletalMesh, int32 L
 			SkeletalMesh->ReserveLODImportData(MaxDependentLODIndex);
 			//Reduce all dependent LODs
 			FThreadSafeBool bNeedsPackageDirtied(false);
-			TBitArray<> CanReduceLODInParallel;
-			CanReduceLODInParallel.Init(true, SkeletalMesh->GetLODNum());
-			//Reduce the LODs which need to save the OriginalReductionSourceMeshData in the main thread, bulkdata are not multithread safe.
-			for (int32 DependentLODIndex : DependentLODs)
+			
+			//Adjust the InlineReductionCacheDatas before simplifying dependent LODs
+			if (SkeletalMesh->GetImportedModel()->InlineReductionCacheDatas.Num() < LODNumber)
 			{
-				const FSkeletalMeshLODInfo* DepLODInfo = SkeletalMesh->GetLODInfo(DependentLODIndex);
-				if (!DepLODInfo || (!DepLODInfo->bHasBeenSimplified || DependentLODIndex == DepLODInfo->ReductionSettings.BaseLOD))
-				{
-					CanReduceLODInParallel[DependentLODIndex] = false;
-					FLODUtilities::SimplifySkeletalMeshLOD(SkeletalMesh, DependentLODIndex, TargetPlatform, false, &bNeedsPackageDirtied);
-				}
+				SkeletalMesh->GetImportedModel()->InlineReductionCacheDatas.AddDefaulted(LODNumber - SkeletalMesh->GetImportedModel()->InlineReductionCacheDatas.Num());
 			}
-			
-			// Reduce LODs in parallel
+			else if (SkeletalMesh->GetImportedModel()->InlineReductionCacheDatas.Num() > LODNumber)
+			{
+				//If we have too much entry simply shrink the array to valid LODModel size
+				SkeletalMesh->GetImportedModel()->InlineReductionCacheDatas.SetNum(LODNumber);
+			}
+
+			// Reduce LODs in parallel (reduction is multithread safe)
 			const bool bHasAccessToLockedProperties = !FSkeletalMeshAsyncBuildScope::ShouldWaitOnLockedProperties(SkeletalMesh);
-			ParallelFor(DependentLODs.Num(), [&DependentLODs, &SkeletalMesh, &bNeedsPackageDirtied, &CanReduceLODInParallel, bHasAccessToLockedProperties, &TargetPlatform](int32 IterationIndex)
+			ParallelFor(DependentLODs.Num(), [&DependentLODs, &SkeletalMesh, &bNeedsPackageDirtied, bHasAccessToLockedProperties, &TargetPlatform](int32 IterationIndex)
 			{
 				TUniquePtr<FSkeletalMeshAsyncBuildScope> AsyncBuildScope(bHasAccessToLockedProperties ? MakeUnique<FSkeletalMeshAsyncBuildScope>(SkeletalMesh) : nullptr);
 
 				check(DependentLODs.IsValidIndex(IterationIndex));
 				int32 DependentLODIndex = DependentLODs[IterationIndex];
-				if (CanReduceLODInParallel[DependentLODIndex])
-				{
-					check(SkeletalMesh->GetLODInfo(DependentLODIndex)); //We cannot add a LOD when reducing with multi thread, so check we already have one
-					FLODUtilities::SimplifySkeletalMeshLOD(SkeletalMesh, DependentLODIndex, TargetPlatform, false, &bNeedsPackageDirtied);
-				}
+				check(SkeletalMesh->GetLODInfo(DependentLODIndex)); //We cannot add a LOD when reducing with multi thread, so check we already have one
+				FLODUtilities::SimplifySkeletalMeshLOD(SkeletalMesh, DependentLODIndex, TargetPlatform, false, &bNeedsPackageDirtied);
 			});
 
 			if (bNeedsPackageDirtied && IsInGameThread())
diff --git a/Engine/Source/Developer/SkeletalMeshUtilitiesCommon/Public/LODUtilities.h b/Engine/Source/Developer/SkeletalMeshUtilitiesCommon/Public/LODUtilities.h
--- a/Engine/Source/Developer/SkeletalMeshUtilitiesCommon/Public/LODUtilities.h
+++ b/Engine/Source/Developer/SkeletalMeshUtilitiesCommon/Public/LODUtilities.h
@@ -6,7 +6,9 @@
 #include "MeshBuild.h"
 #include "Engine/SkeletalMesh.h"
 #include "Rendering/SkeletalMeshLODImporterData.h"
+#include "Rendering/SkeletalMeshLODModel.h"
 #include "Framework/Commands/UIAction.h"
+#include "Animation/MorphTarget.h"
 
 namespace ClothingAssetUtils
 {
@@ -25,6 +27,14 @@ struct FSkeletalMeshUpdateContext
 	FExecuteAction				OnLODChanged;
 };
 
+/* Helper struct to define inline data when applying morph target on a reduce LOD. Inline data is needed when reduction is inline (LOD reduce itself). */
+struct FInlineReductionDataParameter
+{
+	bool bIsDataValid = false;
+	FSkeletalMeshLODModel InlineOriginalSrcModel;
+	TMap<FString, TArray<FMorphTargetDelta>> InlineOriginalSrcMorphTargetData;
+};
+
 class FSkeletalMeshImportData;
 
 //////////////////////////////////////////////////////////////////////////
@@ -85,7 +95,7 @@ public:
 	* @param LodIndex - The LOD index to restore the imported LOD model
 	* @param bReregisterComponent - if true the component using the skeletal mesh will all be re register.
 	*/
-	static bool RestoreSkeletalMeshLODImportedData(USkeletalMesh* SkeletalMesh, int32 LodIndex);
+	static bool RestoreSkeletalMeshLODImportedData_DEPRECATED(USkeletalMesh* SkeletalMesh, int32 LodIndex);
 	
 	/**
 	 * Refresh LOD Change
@@ -185,7 +195,7 @@ private:
 	* @param SourceLOD      - The source LOD morph target .
 	* @param DestinationLOD   - The destination LOD morph target to apply the source LOD morph target
 	*/
-	static void ApplyMorphTargetsToLOD(USkeletalMesh* SkeletalMesh, int32 SourceLOD, int32 DestinationLOD);
+	static void ApplyMorphTargetsToLOD(USkeletalMesh* SkeletalMesh, int32 SourceLOD, int32 DestinationLOD, const FInlineReductionDataParameter& InlineApplyMorphTargetParameter);
 
 	/**
 	*  Clear generated morphtargets for the given LODs
diff --git a/Engine/Source/Editor/Persona/Private/PersonaMeshDetails.cpp b/Engine/Source/Editor/Persona/Private/PersonaMeshDetails.cpp
--- a/Engine/Source/Editor/Persona/Private/PersonaMeshDetails.cpp
+++ b/Engine/Source/Editor/Persona/Private/PersonaMeshDetails.cpp
@@ -2721,13 +2721,11 @@ void FPersonaMeshDetails::AddLODLevelCategories(IDetailLayoutBuilder& DetailLayo
 				
 				bool bIsbuildAvailable = SkelMesh->IsLODImportedDataBuildAvailable(LODIndex);
 
-				bool bIsReductionDataPresent = (SkelMesh->GetImportedModel()->OriginalReductionSourceMeshData.IsValidIndex(LODIndex) && !SkelMesh->GetImportedModel()->OriginalReductionSourceMeshData[LODIndex]->IsEmpty());
 				
 				//Avoid offering re-generate if the LOD is reduce on himself and do not have the original data, the user in this case has to re-import the asset to generate the data 
 				bool LodCannotRegenerate = (SkelMesh->GetLODInfo(LODIndex) != nullptr
 					&& LODIndex == SkelMesh->GetLODInfo(LODIndex)->ReductionSettings.BaseLOD
 					&& SkelMesh->GetLODInfo(LODIndex)->bHasBeenSimplified
-					&& !bIsReductionDataPresent
 					&& !bIsbuildAvailable);
 
 				bool bShowGenerateButtons = !LodCannotRegenerate;
@@ -3352,8 +3350,7 @@ void FPersonaMeshDetails::RestoreNonReducedLOD(int32 LODIndex)
 
 	if (CurrentLODInfo->bHasBeenSimplified
 		&& !bIsReductionActive
-		&& (bIsLODModelbuildDataAvailable
-			|| FLODUtilities::RestoreSkeletalMeshLODImportedData(SkelMesh, LODIndex)))
+		&& bIsLODModelbuildDataAvailable)
 	{
 		CurrentLODInfo->bHasBeenSimplified = false;
 	}
@@ -3447,18 +3444,13 @@ void FPersonaMeshDetails::RegenerateOneLOD(int32 LODIndex)
 		FSkeletalMeshLODInfo& CurrentLODInfo = *(SkelMesh->GetLODInfo(LODIndex));
 		
 		bool bIsLODModelbuildDataAvailable = SkelMesh->GetImportedModel()->LODModels.IsValidIndex(LODIndex) && SkelMesh->IsLODImportedDataBuildAvailable(LODIndex);
-		bool bIsReductionDataPresent = (SkelMesh->GetImportedModel()->OriginalReductionSourceMeshData.IsValidIndex(LODIndex) && !SkelMesh->GetImportedModel()->OriginalReductionSourceMeshData[LODIndex]->IsEmpty());
 		if (LODIndex == CurrentLODInfo.ReductionSettings.BaseLOD
 			&& CurrentLODInfo.bHasBeenSimplified
 			&& !SkelMesh->IsReductionActive(LODIndex)
-			&& (bIsLODModelbuildDataAvailable || bIsReductionDataPresent))
+			&& (bIsLODModelbuildDataAvailable))
 		{
-			//Restore the base LOD data
+			//The build of the skeletalmesh will restore the data properly
 			CurrentLODInfo.bHasBeenSimplified = false;
-			if (!bIsLODModelbuildDataAvailable)
-			{
-				FLODUtilities::RestoreSkeletalMeshLODImportedData(SkelMesh, LODIndex);
-			}
 			return;
 		}
 		else if (!CurrentLODInfo.bHasBeenSimplified
@@ -3473,6 +3465,7 @@ void FPersonaMeshDetails::RegenerateOneLOD(int32 LODIndex)
 		UpdateContext.AssociatedComponents.Push(GetPersonaToolkit()->GetPreviewMeshComponent());
 
 		FLODUtilities::SimplifySkeletalMeshLOD(UpdateContext, LODIndex, GetTargetPlatformManagerRef().GetRunningTargetPlatform());
+
 	}
 	return;
 }
diff --git a/Engine/Source/Editor/UnrealEd/Private/Fbx/FbxSkeletalMeshImport.cpp b/Engine/Source/Editor/UnrealEd/Private/Fbx/FbxSkeletalMeshImport.cpp
--- a/Engine/Source/Editor/UnrealEd/Private/Fbx/FbxSkeletalMeshImport.cpp
+++ b/Engine/Source/Editor/UnrealEd/Private/Fbx/FbxSkeletalMeshImport.cpp
@@ -1881,6 +1881,17 @@ USkeletalMesh* UnFbx::FFbxImporter::ImportSkeletalMesh(FImportSkeletalMeshArgs &
 			// Create actual rendering data.
 			bBuildSuccess = MeshUtilities.BuildSkeletalMesh(ImportedResource->LODModels[ImportLODModelIndex], SkeletalMesh->GetPathName(), SkeletalMesh->GetRefSkeleton(), LODInfluences, LODWedges, LODFaces, LODPoints, LODPointToRawMap, LegacyBuildOptions, &WarningMessages, &WarningNames);
 
+			//Cache the vertex/triangle count in the InlineReductionCacheData so we can know if the LODModel need reduction or not.
+			TArray<FInlineReductionCacheData>& InlineReductionCacheDatas = ImportedResource->InlineReductionCacheDatas;
+			if (!InlineReductionCacheDatas.IsValidIndex(ImportLODModelIndex))
+			{
+				InlineReductionCacheDatas.AddDefaulted((ImportLODModelIndex + 1) - InlineReductionCacheDatas.Num());
+			}
+			if (ensure(InlineReductionCacheDatas.IsValidIndex(ImportLODModelIndex)))
+			{
+				InlineReductionCacheDatas[ImportLODModelIndex].SetCacheGeometryInfo(ImportedResource->LODModels[ImportLODModelIndex]);
+			}
+
 			// temporary hack of message/names, should be one token or a struct
 			if (WarningMessages.Num() > 0 && WarningNames.Num() == WarningMessages.Num())
 			{
diff --git a/Engine/Source/Editor/UnrealEd/Private/ImportUtils/SkeletalMeshImportUtils.cpp b/Engine/Source/Editor/UnrealEd/Private/ImportUtils/SkeletalMeshImportUtils.cpp
--- a/Engine/Source/Editor/UnrealEd/Private/ImportUtils/SkeletalMeshImportUtils.cpp
+++ b/Engine/Source/Editor/UnrealEd/Private/ImportUtils/SkeletalMeshImportUtils.cpp
@@ -228,68 +228,6 @@ void SkeletalMesUtilsImpl::SaveSkeletalMeshLODModelSections(USkeletalMesh* Sourc
 	const FSkeletalMeshLODModel* SourceLODModel = &SourceMeshModel->LODModels[LodIndex];
 	FSkeletalMeshLODModel OriginalLODModel;
 
-	if (bSaveNonReducedMeshData && (SourceMeshModel->OriginalReductionSourceMeshData.IsValidIndex(LodIndex) && !SourceMeshModel->OriginalReductionSourceMeshData[LodIndex]->IsEmpty()))
-	{
-		TMap<FString, TArray<FMorphTargetDelta>> TempLODMorphTargetData;
-		//Get the before reduce LODModel, this lod model contain all the possible sections
-		SourceMeshModel->OriginalReductionSourceMeshData[LodIndex]->LoadReductionData(OriginalLODModel, TempLODMorphTargetData, SourceSkeletalMesh);
-		//If there was section that was remove by the reduction (Disabled in the original data, zero triangle after reduction, GenerateUpTo settings...),
-		//we have to use the original section data and apply the section data that was modified after the reduction
-		if (OriginalLODModel.Sections.Num() > SourceLODModel->Sections.Num())
-		{
-			TArray<bool> OriginalMatched;
-			OriginalMatched.AddZeroed(OriginalLODModel.Sections.Num());
-			//Now apply the after reduce settings change, but we need to match the section since there can be reduced one
-			for (int32 ReduceSectionIndex = 0; ReduceSectionIndex < SourceLODModel->Sections.Num(); ++ReduceSectionIndex)
-			{
-				const FSkelMeshSection& ReduceSection = SourceLODModel->Sections[ReduceSectionIndex];
-				for (int32 OriginalSectionIndex = 0; OriginalSectionIndex < OriginalLODModel.Sections.Num(); ++OriginalSectionIndex)
-				{
-					if (OriginalMatched[OriginalSectionIndex])
-					{
-						continue;
-					}
-					FSkelMeshSection& OriginalSection = OriginalLODModel.Sections[OriginalSectionIndex];
-					if ((OriginalSection.bDisabled) || (OriginalSection.GenerateUpToLodIndex != INDEX_NONE && OriginalSection.GenerateUpToLodIndex < LodIndex))
-					{
-						continue;
-					}
-
-					if (ReduceSection.MaterialIndex == OriginalSection.MaterialIndex)
-					{
-						OriginalMatched[OriginalSectionIndex] = true;
-						OriginalSection.bDisabled = ReduceSection.bDisabled;
-						OriginalSection.bCastShadow = ReduceSection.bCastShadow;
-						OriginalSection.bVisibleInRayTracing = ReduceSection.bVisibleInRayTracing;
-						OriginalSection.bRecomputeTangent = ReduceSection.bRecomputeTangent;
-						OriginalSection.RecomputeTangentsVertexMaskChannel = ReduceSection.RecomputeTangentsVertexMaskChannel;
-						OriginalSection.GenerateUpToLodIndex = ReduceSection.GenerateUpToLodIndex;
-						break;
-					}
-				}
-			}
-			//Set the unmatched original section data using the current UserSectionsData so we keep the user changes
-			for (int32 OriginalSectionIndex = 0; OriginalSectionIndex < OriginalLODModel.Sections.Num(); ++OriginalSectionIndex)
-			{
-				if (OriginalMatched[OriginalSectionIndex])
-				{
-					continue;
-				}
-				FSkelMeshSection& OriginalSection = OriginalLODModel.Sections[OriginalSectionIndex];
-				if (const FSkelMeshSourceSectionUserData* ReduceUserSectionData = SourceLODModel->UserSectionsData.Find(OriginalSection.OriginalDataSectionIndex))
-				{
-					OriginalSection.bDisabled = ReduceUserSectionData->bDisabled;
-					OriginalSection.bCastShadow = ReduceUserSectionData->bCastShadow;
-					OriginalSection.bVisibleInRayTracing = ReduceUserSectionData->bVisibleInRayTracing;
-					OriginalSection.bRecomputeTangent = ReduceUserSectionData->bRecomputeTangent;
-					OriginalSection.RecomputeTangentsVertexMaskChannel = ReduceUserSectionData->RecomputeTangentsVertexMaskChannel;
-					OriginalSection.GenerateUpToLodIndex = ReduceUserSectionData->GenerateUpToLodIndex;
-				}
-			}
-			//Use the OriginalLODModel
-			SourceLODModel = &OriginalLODModel;
-		}
-	}
 	ExistingMeshDataPtr->ExistingImportMeshLodSectionMaterialData.AddZeroed();
 	check(ExistingMeshDataPtr->ExistingImportMeshLodSectionMaterialData.IsValidIndex(LodIndex));
 
@@ -407,6 +345,8 @@ TSharedPtr<FExistingSkelMeshData> SkeletalMeshImportUtils::SaveExistingSkelMeshD
 	}
 	ExistingMeshDataPtr->ExistingRetargetBasePose = SourceSkeletalMesh->GetRetargetBasePose();
 
+	ExistingMeshDataPtr->ExistingInlineReductionCacheDatas = SourceMeshModel->InlineReductionCacheDatas;
+
 	if (SourceMeshModel->LODModels.Num() > 0 &&
 		SourceSkeletalMesh->GetLODNum() == SourceMeshModel->LODModels.Num())
 	{
@@ -415,19 +355,6 @@ TSharedPtr<FExistingSkelMeshData> SkeletalMeshImportUtils::SaveExistingSkelMeshD
 		ExistingMeshDataPtr->ExistingLODModels.Empty(SourceMeshModel->LODModels.Num());
 		for ( int32 LODIndex = 0; LODIndex < SourceMeshModel->LODModels.Num() ; ++LODIndex)
 		{
-			//const int32 ReductionLODIndex = LODModelIndex + OffsetReductionLODIndex;
-			TSharedPtr<FReductionBaseSkeletalMeshBulkData> ReductionLODData;
-			if (SourceMeshModel->OriginalReductionSourceMeshData.IsValidIndex(LODIndex) && !SourceMeshModel->OriginalReductionSourceMeshData[LODIndex]->IsEmpty())
-			{
-				FSkeletalMeshLODModel BaseLODModel;
-				TMap<FString, TArray<FMorphTargetDelta>> BaseLODMorphTargetData;
-				SourceMeshModel->OriginalReductionSourceMeshData[LODIndex]->LoadReductionData(BaseLODModel, BaseLODMorphTargetData, SourceSkeletalMesh);
-				ReductionLODData = MakeShared<FReductionBaseSkeletalMeshBulkData>();
-				ReductionLODData->SaveReductionData(BaseLODModel, BaseLODMorphTargetData, SourceSkeletalMesh);
-			}
-			//Add the reduction source mesh data if it exist, otherwise an empty sharedPtr.
-			ExistingMeshDataPtr->ExistingOriginalReductionSourceMeshData.Add(MoveTemp(ReductionLODData));
-			
 			//Add a new LOD Model to the existing LODModels data
 			const FSkeletalMeshLODModel& LODModel = SourceMeshModel->LODModels[LODIndex];
 			ExistingMeshDataPtr->ExistingLODModels.Add(FSkeletalMeshLODModel::CreateCopy(&LODModel));
@@ -751,7 +678,6 @@ void SkeletalMeshImportUtils::RestoreExistingSkelMeshData(const TSharedPtr<const
 			bMaterialReset = true;
 		}
 	}
-
 	SkeletalMesh->SetLODSettings(MeshData->ExistingLODSettings);
 	// ensure LOD 0 contains correct setting 
 	if (SkeletalMesh->GetLODSettings() && SkeletalMesh->GetLODInfoArray().Num() > 0)
@@ -779,31 +705,14 @@ void SkeletalMeshImportUtils::RestoreExistingSkelMeshData(const TSharedPtr<const
 				SkeletalMesh->GetMeshOnlySocketList().Add(MeshData->ExistingSockets[i]);
 			}
 		}
+		if (!ensure(SkeletalMeshImportedModel->InlineReductionCacheDatas.IsValidIndex(0)))
+		{
+			UE_LOG(LogSkeletalMeshImport, Display, TEXT("Reimport SkeletalMesh did not have a valid inline reduction cache data for LOD 0"));
+		}
 
 		// We copy back and fix-up the LODs that still work with this skeleton.
 		if (MeshData->ExistingLODModels.Num() > 1)
 		{
-			auto RestoreReductionSourceData = [&SkeletalMesh, &SkeletalMeshImportedModel, &MeshData](int32 ExistingIndex, int32 NewIndex)
-			{
-				if (!MeshData->ExistingOriginalReductionSourceMeshData[ExistingIndex].IsValid() || MeshData->ExistingOriginalReductionSourceMeshData[ExistingIndex]->IsEmpty())
-				{
-					return;
-				}
-				//Restore the original reduction source mesh data
-				FSkeletalMeshLODModel BaseLODModel;
-				TMap<FString, TArray<FMorphTargetDelta>> BaseLODMorphTargetData;
-				MeshData->ExistingOriginalReductionSourceMeshData[ExistingIndex]->LoadReductionData(BaseLODModel, BaseLODMorphTargetData, SkeletalMesh);
-				FReductionBaseSkeletalMeshBulkData* ReductionLODData = new FReductionBaseSkeletalMeshBulkData();
-				ReductionLODData->SaveReductionData(BaseLODModel, BaseLODMorphTargetData, SkeletalMesh);
-				//Add necessary empty slot
-				while (SkeletalMeshImportedModel->OriginalReductionSourceMeshData.Num() < NewIndex)
-				{
-					FReductionBaseSkeletalMeshBulkData* EmptyReductionLODData = new FReductionBaseSkeletalMeshBulkData();
-					SkeletalMeshImportedModel->OriginalReductionSourceMeshData.Add(EmptyReductionLODData);
-				}
-				SkeletalMeshImportedModel->OriginalReductionSourceMeshData.Add(ReductionLODData);
-			};
-
 			if (SkeletonsAreCompatible(SkeletalMesh->GetRefSkeleton(), MeshData->ExistingRefSkeleton, bImportSkinningOnly))
 			{
 				// First create mapping table from old skeleton to new skeleton.
@@ -912,7 +821,35 @@ void SkeletalMeshImportUtils::RestoreExistingSkelMeshData(const TSharedPtr<const
 						//We need to add LODInfo
 						SkeletalMeshImportedModel->LODModels.Add(LODModelCopy);
 						SkeletalMesh->AddLODInfo(LODInfo);
-						RestoreReductionSourceData(LODIndex, SkeletalMesh->GetLODNum() - 1);
+						
+						auto FillInlineReductionData = [&MeshData, &SkeletalMeshImportedModel, LODIndex]()
+						{
+							if (MeshData->ExistingInlineReductionCacheDatas.IsValidIndex(LODIndex))
+							{
+								SkeletalMeshImportedModel->InlineReductionCacheDatas[LODIndex] = MeshData->ExistingInlineReductionCacheDatas[LODIndex];
+							}
+							else
+							{
+								SkeletalMeshImportedModel->InlineReductionCacheDatas[LODIndex].SetCacheGeometryInfo(SkeletalMeshImportedModel->LODModels[LODIndex]);
+							}
+						};
+
+						if (!SkeletalMeshImportedModel->InlineReductionCacheDatas.IsValidIndex(LODIndex))
+						{
+							SkeletalMeshImportedModel->InlineReductionCacheDatas.AddDefaulted(LODIndex + 1 - SkeletalMeshImportedModel->InlineReductionCacheDatas.Num());
+							check(SkeletalMeshImportedModel->InlineReductionCacheDatas.IsValidIndex(LODIndex));
+							FillInlineReductionData();
+						}
+						else
+						{
+							uint32 LODVertexCount = MAX_uint32;
+							uint32 LODTriangleCount = MAX_uint32;
+							SkeletalMeshImportedModel->InlineReductionCacheDatas[LODIndex].GetCacheGeometryInfo(LODVertexCount, LODTriangleCount);
+							if (LODVertexCount == MAX_uint32 || LODTriangleCount == MAX_uint32)
+							{
+								FillInlineReductionData();
+							}
+						}
 					}
 				}
 			}
@@ -990,14 +927,6 @@ void SkeletalMeshImportUtils::RestoreExistingSkelMeshData(const TSharedPtr<const
 		SkeletalMesh->AddAssetUserData(UserDataObject);
 	}
 
-	if (!bImportSkinningOnly && (!MeshData->ExistingLODInfo.IsValidIndex(SafeReimportLODIndex) || !MeshData->ExistingLODInfo[SafeReimportLODIndex].bHasBeenSimplified))
-	{
-		if (SkeletalMeshImportedModel->OriginalReductionSourceMeshData.IsValidIndex(SafeReimportLODIndex))
-		{
-			SkeletalMeshImportedModel->OriginalReductionSourceMeshData[SafeReimportLODIndex]->EmptyBulkData();
-		}
-	}
-
 	//Copy mesh changed delegate data
 	SkeletalMesh->GetOnMeshChanged() = MeshData->ExistingOnMeshChanged;
 
diff --git a/Engine/Source/Editor/UnrealEd/Public/ImportUtils/SkelImport.h b/Engine/Source/Editor/UnrealEd/Public/ImportUtils/SkelImport.h
--- a/Engine/Source/Editor/UnrealEd/Public/ImportUtils/SkelImport.h
+++ b/Engine/Source/Editor/UnrealEd/Public/ImportUtils/SkelImport.h
@@ -43,7 +43,7 @@ struct FExistingMeshLodSectionData
 struct FExistingSkelMeshData
 {
 	TArray<USkeletalMeshSocket*>			ExistingSockets;
-	TArray<TSharedPtr<FReductionBaseSkeletalMeshBulkData>> ExistingOriginalReductionSourceMeshData;
+	TArray<FInlineReductionCacheData>		ExistingInlineReductionCacheDatas;
 	TIndirectArray<FSkeletalMeshLODModel>	ExistingLODModels;
 	TArray<FSkeletalMeshLODInfo>			ExistingLODInfo;
 	FReferenceSkeleton						ExistingRefSkeleton;
diff --git a/Engine/Source/Runtime/Core/Public/UObject/UE5MainStreamObjectVersion.h b/Engine/Source/Runtime/Core/Public/UObject/UE5MainStreamObjectVersion.h
--- a/Engine/Source/Runtime/Core/Public/UObject/UE5MainStreamObjectVersion.h
+++ b/Engine/Source/Runtime/Core/Public/UObject/UE5MainStreamObjectVersion.h
@@ -180,6 +180,9 @@ struct CORE_API FUE5MainStreamObjectVersion
 		
 		// Add custom version to FDynamicMesh3
 		DynamicMeshCompactedSerialization,
+
+		// Remove the inline reduction bulkdata and replace it by a simple vertex and triangle count cache
+		ConvertReductionBaseSkeletalMeshBulkDataToInlineReductionCacheData,
 		
 		// -----<new versions can be added above this line>-------------------------------------------------
 		VersionPlusOne,
diff --git a/Engine/Source/Runtime/Engine/Classes/Engine/SkeletalMesh.h b/Engine/Source/Runtime/Engine/Classes/Engine/SkeletalMesh.h
--- a/Engine/Source/Runtime/Engine/Classes/Engine/SkeletalMesh.h
+++ b/Engine/Source/Runtime/Engine/Classes/Engine/SkeletalMesh.h
@@ -150,7 +150,6 @@ public:
 class FSkeletalMeshPostLoadContext : public FSkeletalMeshCompilationContext
 {
 public:
-	bool bIsPreSkeletalMeshBuildRefactor = false;
 	bool bHasCachedDerivedData = false;
 };
 
@@ -787,7 +786,7 @@ public:
 
 	/* Fill the OutMesh with the imported data */
 	void LoadLODImportedData(const int32 LODIndex, FSkeletalMeshImportData& OutMesh) const;
-	
+
 	/* Fill the asset LOD entry with the InMesh. */
 	void SaveLODImportedData(const int32 LODIndex, FSkeletalMeshImportData& InMesh);
 	
@@ -1014,6 +1013,9 @@ public:
 		return false;
 	}
 
+	/* Return true if this skeletalmesh was never build since its creation. */
+	bool IsInitialBuildDone() const;
+
 	/* Return true if the reduction settings are setup to reduce a LOD*/
 	bool IsReductionActive(int32 LODIndex) const;
 
diff --git a/Engine/Source/Runtime/Engine/Private/Rendering/SkeletalMeshLODImporterData.cpp b/Engine/Source/Runtime/Engine/Private/Rendering/SkeletalMeshLODImporterData.cpp
--- a/Engine/Source/Runtime/Engine/Private/Rendering/SkeletalMeshLODImporterData.cpp
+++ b/Engine/Source/Runtime/Engine/Private/Rendering/SkeletalMeshLODImporterData.cpp
@@ -587,6 +587,36 @@ void FReductionBaseSkeletalMeshBulkData::GetGeometryInfo(uint32& LODVertexNumber
 	LODTriNumber = CacheLODTriNumber;
 }
 
+/*------------------------------------------------------------------------------
+FInlineReductionCacheData
+------------------------------------------------------------------------------*/
+
+void FInlineReductionCacheData::SetCacheGeometryInfo(const FSkeletalMeshLODModel& SourceLODModel)
+{
+	CacheLODVertexCount = 0;
+	CacheLODTriCount = 0;
+	for (int32 SectionIndex = 0; SectionIndex < SourceLODModel.Sections.Num(); ++SectionIndex)
+	{
+		const FSkelMeshSection& Section = SourceLODModel.Sections[SectionIndex];
+		//We count disabled section, since the render buffer contain the disabled section data. This is crucial for memory budget
+		//Make sure the count fit in a uint32
+		CacheLODVertexCount += Section.NumVertices < 0 ? 0 : Section.NumVertices;
+		CacheLODTriCount += Section.NumTriangles;
+	}
+}
+
+void FInlineReductionCacheData::SetCacheGeometryInfo(uint32 LODVertexCount, uint32 LODTriCount)
+{
+	CacheLODVertexCount = LODVertexCount;
+	CacheLODTriCount = LODTriCount;
+}
+
+void FInlineReductionCacheData::GetCacheGeometryInfo(uint32& LODVertexCount, uint32& LODTriCount) const
+{
+	LODVertexCount = CacheLODVertexCount;
+	LODTriCount = CacheLODTriCount;
+}
+
 /*------------------------------------------------------------------------------
 FRawSkeletalMeshBulkData
 ------------------------------------------------------------------------------*/
@@ -792,7 +822,6 @@ void FRawSkeletalMeshBulkData::SaveRawMesh(FSkeletalMeshImportData& InMesh)
 		// Preserve CustomVersions at save time so we can reuse the same ones when reloading direct from memory
 		SerializeLoadingCustomVersionContainer = Ar.GetCustomVersions();
 	}
-	// Unlock bulk data when we leave scope
 
 	//Create the guid from the content, this allow to use the data into the ddc key
 	FSHA1 Sha;
diff --git a/Engine/Source/Runtime/Engine/Private/SkeletalMesh.cpp b/Engine/Source/Runtime/Engine/Private/SkeletalMesh.cpp
--- a/Engine/Source/Runtime/Engine/Private/SkeletalMesh.cpp
+++ b/Engine/Source/Runtime/Engine/Private/SkeletalMesh.cpp
@@ -33,6 +33,7 @@
 #include "UObject/PropertyPortFlags.h"
 #include "UObject/RenderingObjectVersion.h"
 #include "UObject/UObjectIterator.h"
+#include "UObject/UE5MainStreamObjectVersion.h"
 #include "EngineUtils.h"
 #include "EditorSupportDelegates.h"
 #include "GPUSkinVertexFactory.h"
@@ -608,6 +609,16 @@ void USkeletalMesh::ValidateBoundsExtension()
 }
 
 #if WITH_EDITOR
+
+bool USkeletalMesh::IsInitialBuildDone() const
+{
+	//We are consider built if we have a valid lod model and a valid inline reduction cache
+	return GetImportedModel() != nullptr &&
+		GetImportedModel()->LODModels.Num() >= 0 &&
+		GetImportedModel()->LODModels[0].Sections.Num() > 0 &&
+		GetImportedModel()->InlineReductionCacheDatas.Num() > 0;
+}
+
 /* Return true if the reduction settings are setup to reduce a LOD*/
 bool USkeletalMesh::IsReductionActive(int32 LODIndex) const
 {
@@ -624,7 +635,7 @@ bool USkeletalMesh::IsReductionActive(int32 LODIndex) const
 		uint32 LODVertexNumber = MAX_uint32;
 		uint32 LODTriNumber = MAX_uint32;
 		const FSkeletalMeshLODInfo* LODInfoPtr = GetLODInfo(LODIndex);
-		bool bLODHasBeenSimplified = LODInfoPtr && LODInfoPtr->bHasBeenSimplified;
+		const bool bLODHasBeenSimplified = LODInfoPtr && LODInfoPtr->bHasBeenSimplified;
 		if (GetImportedModel() && GetImportedModel()->LODModels.IsValidIndex(LODIndex))
 		{
 			if (!bLODHasBeenSimplified)
@@ -642,12 +653,10 @@ bool USkeletalMesh::IsReductionActive(int32 LODIndex) const
 					LODTriNumber += Section.NumTriangles;
 				}
 			}
-			else if (GetImportedModel()->OriginalReductionSourceMeshData.IsValidIndex(LODIndex)
-				&& !GetImportedModel()->OriginalReductionSourceMeshData[LODIndex]->IsEmpty())
+			else if (GetImportedModel()->InlineReductionCacheDatas.IsValidIndex(LODIndex))
 			{
-				//In this case we have to use the stored reduction source data to know how many vertices/triangles we have before the reduction
-				USkeletalMesh* MutableSkeletalMesh = const_cast<USkeletalMesh*>(this);
-				GetImportedModel()->OriginalReductionSourceMeshData[LODIndex]->GetGeometryInfo(LODVertexNumber, LODTriNumber, MutableSkeletalMesh);
+				//In this case we have to use the inline cache reduction data to know how many vertices/triangles we have before the reduction
+				GetImportedModel()->InlineReductionCacheDatas[LODIndex].GetCacheGeometryInfo(LODVertexNumber, LODTriNumber);
 			}
 		}
 		bReductionActive = ReductionModule->IsReductionActive(ReductionSettings, LODVertexNumber, LODTriNumber);
@@ -1702,6 +1711,7 @@ void USkeletalMesh::Serialize( FArchive& Ar )
 	Ar.UsingCustomVersion(FRenderingObjectVersion::GUID);
 	Ar.UsingCustomVersion(FFortniteMainBranchObjectVersion::GUID);
 	Ar.UsingCustomVersion(FNiagaraObjectVersion::GUID);
+	Ar.UsingCustomVersion(FUE5MainStreamObjectVersion::GUID);
 
 	FStripDataFlags StripFlags( Ar );
 
@@ -2160,8 +2170,7 @@ void USkeletalMesh::Build()
 	BeginBuildInternal(Context);
 	
 	// Inline reduction is not thread-safe, prevent async build until this is fixed.
-	const bool bHasInlineReductions = SkeletalMeshImpl::HasInlineReductions(this);
-	if (!bHasInlineReductions && FSkeletalMeshCompilingManager::Get().IsAsyncCompilationAllowed(this))
+	if (FSkeletalMeshCompilingManager::Get().IsAsyncCompilationAllowed(this))
 	{
 		PrepareForAsyncCompilation();
 
@@ -2189,6 +2198,14 @@ void USkeletalMesh::BeginBuildInternal(FSkeletalMeshBuildContext& Context)
 	// Release the static mesh's resources.
 	ReleaseResources();
 
+	//Make sure InlineReduction structure are created
+	const int32 MaxLodIndex = GetLODNum() - 1;
+	if (GetImportedModel()->InlineReductionCacheDatas.IsValidIndex(MaxLodIndex))
+	{
+		//We should not do that in main thread, this is why there is an ensure
+		GetImportedModel()->InlineReductionCacheDatas.AddDefaulted((MaxLodIndex + 1) - GetImportedModel()->InlineReductionCacheDatas.Num());
+	}
+
 	// Flush the resource release commands to the rendering thread to ensure that the build doesn't occur while a resource is still
 	// allocated, and potentially accessing the USkeletalMesh.
 	ReleaseResourcesFence.Wait();
@@ -2707,8 +2724,8 @@ void USkeletalMesh::CreateUserSectionsDataForLegacyAssets()
 
 		bool bMustUseReductionSourceData = bIsLODReductionActive
 			&& ThisLODInfo->bHasBeenSimplified
-			&& GetImportedModel()->OriginalReductionSourceMeshData.IsValidIndex(LodIndex)
-			&& !(GetImportedModel()->OriginalReductionSourceMeshData[LodIndex]->IsEmpty());
+			&& GetImportedModel()->OriginalReductionSourceMeshData_DEPRECATED.IsValidIndex(LodIndex)
+			&& !(GetImportedModel()->OriginalReductionSourceMeshData_DEPRECATED[LodIndex]->IsEmpty());
 
 		if (bIsLODReductionActive && !ThisLODInfo->bHasBeenSimplified && IsLODImportedDataEmpty(LodIndex))
 		{
@@ -2747,7 +2764,7 @@ void USkeletalMesh::CreateUserSectionsDataForLegacyAssets()
 			//We must load the reduction source model, since reduction can remove section
 			FSkeletalMeshLODModel ReductionSrcLODModel;
 			TMap<FString, TArray<FMorphTargetDelta>> TmpMorphTargetData;
-			GetImportedModel()->OriginalReductionSourceMeshData[LodIndex]->LoadReductionData(ReductionSrcLODModel, TmpMorphTargetData, this);
+			GetImportedModel()->OriginalReductionSourceMeshData_DEPRECATED[LodIndex]->LoadReductionData(ReductionSrcLODModel, TmpMorphTargetData, this);
 			
 			//Fill the user data with the original value
 			TMap<int32, FSkelMeshSourceSectionUserData> BackupUserSectionsData = ThisLODModel.UserSectionsData;
@@ -2797,61 +2814,85 @@ void USkeletalMesh::PostLoadValidateUserSectionData()
 			continue;
 		}
 
+		const int32 ReductionBaseLOD = LODInfoPtr->ReductionSettings.BaseLOD;
+		if (!GetImportedModel()->LODModels.IsValidIndex(ReductionBaseLOD))
+		{
+			//The base LOD should always be valid for generated LOD
+			UE_ASSET_LOG(LogSkeletalMesh, Display, this, TEXT("This asset generated lod %d, is base on an invalid LOD index %d."), LodIndex, ReductionBaseLOD);
+			continue;
+		}
+		FSkeletalMeshLODModel& BaseReductionLODModel = GetImportedModel()->LODModels[ReductionBaseLOD];
 		FSkeletalMeshLODModel& ThisLODModel = GetImportedModel()->LODModels[LodIndex];
 		const int32 SectionNum = ThisLODModel.Sections.Num();
-		//See if more then one section use the same UserSectionData
-		bool bLODHaveSectionIssue = false;
-		TBitArray<> AvailableUserSectionData;
-		AvailableUserSectionData.Init(true, ThisLODModel.UserSectionsData.Num());
-		for (int32 SectionIndex = 0; SectionIndex < SectionNum; ++SectionIndex)
+		
+		//We must make sure the result is similar to what the reduction will give. So we will not have more user section data then the number we have for the base LOD.
+		//Because reduction reset the UserSectionData to the number of parent section after the reduction.
+		bool UserSectionDataBadCount = ThisLODModel.UserSectionsData.Num() > SectionNum;
+		if (!UserSectionDataBadCount && LodIndex != ReductionBaseLOD)
 		{
-			FSkelMeshSection& Section = ThisLODModel.Sections[SectionIndex];
-			if (Section.ChunkedParentSectionIndex != INDEX_NONE)
+			UserSectionDataBadCount = (ThisLODModel.UserSectionsData.Num() > BaseReductionLODModel.UserSectionsData.Num());
+		}
+		bool bLODHaveSectionIssue = UserSectionDataBadCount;
+		if (!bLODHaveSectionIssue)
+		{
+			//See if more then one section use the same UserSectionData
+			TBitArray<> AvailableUserSectionData;
+			AvailableUserSectionData.Init(true, ThisLODModel.UserSectionsData.Num());
+			for (int32 SectionIndex = 0; SectionIndex < SectionNum; ++SectionIndex)
 			{
-				continue;
+				FSkelMeshSection& Section = ThisLODModel.Sections[SectionIndex];
+				if (Section.ChunkedParentSectionIndex != INDEX_NONE)
+				{
+					continue;
+				}
+				if (!AvailableUserSectionData.IsValidIndex(Section.OriginalDataSectionIndex) || !AvailableUserSectionData[Section.OriginalDataSectionIndex])
+				{
+					bLODHaveSectionIssue = true;
+					break;
+				}
+				AvailableUserSectionData[Section.OriginalDataSectionIndex] = false;
 			}
-			if(!AvailableUserSectionData.IsValidIndex(Section.OriginalDataSectionIndex) || !AvailableUserSectionData[Section.OriginalDataSectionIndex])
+			if (!bLODHaveSectionIssue)
 			{
-				bLODHaveSectionIssue = true;
-				break;
+				//Everything is good nothing to fix
+				continue;
 			}
-			AvailableUserSectionData[Section.OriginalDataSectionIndex] = false;
-		}
-		if(!bLODHaveSectionIssue)
-		{
-			//Everything is good nothing to fix
-			continue;
 		}
 
 		//Force the source UserSectionData, then restore the UserSectionData value each section was using
 		//We use the source section user data entry in case we do not have any override
-		const FSkeletalMeshLODModel& BaseLODModel = GetImportedModel()->LODModels[LODInfoPtr->ReductionSettings.BaseLOD];
 		TMap<int32, FSkelMeshSourceSectionUserData> NewUserSectionsData;
 
-		int32 CurrentOriginalSectionIndex = 0;
+		int32 CurrentOriginalSectionIndex = -1;
 		for (int32 SectionIndex = 0; SectionIndex < SectionNum; ++SectionIndex)
 		{
 			FSkelMeshSection& Section = ThisLODModel.Sections[SectionIndex];
 			if (Section.ChunkedParentSectionIndex != INDEX_NONE)
 			{
+				//The section zero must never be a chunked children
+				if (!ensure(CurrentOriginalSectionIndex >= 0))
+				{
+					CurrentOriginalSectionIndex = 0;
+				}
 				//We do not restore user section data for chunked section, the parent has already fix it
 				Section.OriginalDataSectionIndex = CurrentOriginalSectionIndex;
 				continue;
 			}
 
+			//Parent (non chunked) section must increment the index
+			CurrentOriginalSectionIndex++;
+
 			FSkelMeshSourceSectionUserData& SectionUserData = NewUserSectionsData.FindOrAdd(CurrentOriginalSectionIndex);
 			if(const FSkelMeshSourceSectionUserData* BackupSectionUserData = ThisLODModel.UserSectionsData.Find(Section.OriginalDataSectionIndex))
 			{
 				SectionUserData = *BackupSectionUserData;
 			}
-			else if(const FSkelMeshSourceSectionUserData* BaseSectionUserData = BaseLODModel.UserSectionsData.Find(CurrentOriginalSectionIndex))
+			else if(const FSkelMeshSourceSectionUserData* BaseSectionUserData = BaseReductionLODModel.UserSectionsData.Find(CurrentOriginalSectionIndex))
 			{
 				SectionUserData = *BaseSectionUserData;
 			}
 
 			Section.OriginalDataSectionIndex = CurrentOriginalSectionIndex;
-			//Parent (non chunked) section must increment the index
-			CurrentOriginalSectionIndex++;
 		}
 		ThisLODModel.UserSectionsData = NewUserSectionsData;
 
@@ -2907,7 +2948,21 @@ void USkeletalMesh::PostLoadEnsureImportDataExist()
 {
 	//If we have a LODModel with no import data and the LOD model have at least one section using more bone then any platform max GPU bone count. We will recreate the import data to allow the asset to be build and chunk properly.
 	const int32 MinimumPerPlatformMaxGPUSkinBones = FGPUBaseSkinVertexFactory::GetMinimumPerPlatformMaxGPUSkinBonesValue();
+	//This flag is set from the max gpu skin bones, if true ImportLODData will be created
 	bool bNeedToCreateImportData = false;
+
+	//The following two flags are acting together when bNeedToCreateImportData == false.
+	//When we found a LOD that is reduce inline, we want to convert the deprecated reduction source data
+	//and create the import data from it.
+
+	//If we found a LOD doing inline reduction this flag will be set to trtue
+	bool bContainInlineReduction = false;
+	//If we found a LOD doing inline reduction but with baked skinweight data this flag will be true.
+	//If this flag is true we will not create the imported data, except if bNeedToCreateData is true
+	//In case bNeedToCreateData is true, user will have to re-import the skinning data to comply with
+	//the per platform max gpu bone.
+	bool bCannotConvertSkinWeightProfile = false;
+
 	for (int32 LodIndex = 0; LodIndex < GetLODNum(); LodIndex++)
 	{
 		const FSkeletalMeshLODModel* LODModel = &(GetImportedModel()->LODModels[LodIndex]);
@@ -2919,6 +2974,9 @@ void USkeletalMesh::PostLoadEnsureImportDataExist()
 		{
 			continue;
 		}
+		
+		//We create the import data only if the mesh was simplified and its doing inline reduction
+		//Non simplified mesh without import data should not be converted because the simplification will kick in and the asset will be simplified (asset vertex and face count will change)
 		const bool bReductionActive = IsReductionActive(LodIndex);
 		const bool bInlineReduction = bReductionActive && (ThisLODInfo->ReductionSettings.BaseLOD == LodIndex);
 		if (bReductionActive && !bInlineReduction)
@@ -2926,6 +2984,26 @@ void USkeletalMesh::PostLoadEnsureImportDataExist()
 			//Generated LOD (not inline) do not need imported data
 			continue;
 		}
+		else if (bInlineReduction && ThisLODInfo->bHasBeenSimplified)
+		{
+			//If we cannot convert alternate skin weight data, we do not want to allow 
+			if (GetImportedModel()->OriginalReductionSourceMeshData_DEPRECATED.IsValidIndex(LodIndex))
+			{
+				//Old inline reduced assets do not have the original imported skin weight data, it must be re-import if we want to convert it to import data
+				bool bUseSkinWeightProfile = LODModel->SkinWeightProfiles.Num() > 0;
+				FSkeletalMeshLODModel ReductionLODModel;
+				TMap<FString, TArray<FMorphTargetDelta>> ReductionLODMorphTargetData;
+				//Swap the LODModel pointer to the one we store when reducing, so we add the true import data
+				GetImportedModel()->OriginalReductionSourceMeshData_DEPRECATED[LodIndex]->LoadReductionData(ReductionLODModel, ReductionLODMorphTargetData, this);
+
+				if (bUseSkinWeightProfile && ReductionLODModel.SkinWeightProfiles.IsEmpty())
+				{
+					//If we cannot convert a particular LOD, prevent Creation of import data and exit the loop
+					bCannotConvertSkinWeightProfile = true;
+				}
+			}
+			bContainInlineReduction = true;
+		}
 		//See if the LODModel data use more bones then the chunking allow
 		int32 MaxBoneperSection = 0;
 		for(const FSkelMeshSection& Section : LODModel->Sections)
@@ -2939,8 +3017,18 @@ void USkeletalMesh::PostLoadEnsureImportDataExist()
 			break;
 		}
 	}
-	if (bNeedToCreateImportData)
+
+	if (bNeedToCreateImportData || (bContainInlineReduction && !bCannotConvertSkinWeightProfile))
 	{
+		if (!ensure(IsInGameThread()))
+		{
+			UE_ASSET_LOG(LogSkeletalMesh, Error, this, TEXT("USkeletalMesh::PostLoadEnsureImportDataExist() must be call on the game thread."));
+			return;
+		}
+		if (bCannotConvertSkinWeightProfile)
+		{
+			UE_ASSET_LOG(LogSkeletalMesh, Warning, this, TEXT("USkeletalMesh::PostLoadEnsureImportDataExist() Cannot convert old alternate skin weight profiles data. Re-import all alternate profiles for this skeletal mesh."));
+		}
 		IMeshUtilities& MeshUtilities = FModuleManager::Get().LoadModuleChecked<IMeshUtilities>("MeshUtilities");
 		//We create the import data for all LOD that do not have import data except for the generated LODs.
 		MeshUtilities.CreateImportDataFromLODModel(this);
@@ -3154,17 +3242,6 @@ void USkeletalMesh::BeginPostLoadInternal(FSkeletalMeshPostLoadContext& Context)
 		MeshClothingAsset->ConditionalPostLoad();
 	}
 
-	Context.bIsPreSkeletalMeshBuildRefactor = GetLinkerCustomVersion(FEditorObjectVersion::GUID) < FEditorObjectVersion::SkeletalMeshBuildRefactor;
-	
-	if (!Context.bIsPreSkeletalMeshBuildRefactor)
-	{
-		check(!GetUseLegacyMeshDerivedDataKey());
-		
-		// We can safely release the write lock now since we're not going to modify this value during build, but we still
-		// require it to be locked in read-only if the game-thread would want to write to it while we read it...
-		ReleaseAsyncProperty(ESkeletalMeshAsyncProperties::UseLegacyMeshDerivedDataKey, EAsyncPropertyLockType::WriteOnly);
-	}
-
 	// Make sure the mesh editor data object is a sub object of the skeletalmesh, rename it to change the owner to be the skeletalmesh.
 	if (IsMeshEditorDataValid() && GetMeshEditorData().GetOuter() != this)
 	{
@@ -3265,6 +3342,47 @@ void USkeletalMesh::BeginPostLoadInternal(FSkeletalMeshPostLoadContext& Context)
 			}
 		}
 	}
+
+	if (GetLinkerCustomVersion(FEditorObjectVersion::GUID) < FEditorObjectVersion::SkeletalMeshMoveEditorSourceDataToPrivateAsset)
+	{
+		ReserveLODImportData(GetImportedModel()->LODModels.Num() - 1);
+		for (int32 LODIndex = 0; LODIndex < GetImportedModel()->LODModels.Num(); ++LODIndex)
+		{
+			FSkeletalMeshLODModel& ThisLODModel = GetImportedModel()->LODModels[LODIndex];
+			//We can have partial data if the asset was save after the split workflow implementation
+			//Use the deprecated member to retrieve this data
+			if (GetLinkerCustomVersion(FFortniteMainBranchObjectVersion::GUID) >= FFortniteMainBranchObjectVersion::NewSkeletalMeshImporterWorkflow)
+			{
+				if (!ThisLODModel.RawSkeletalMeshBulkData_DEPRECATED.IsEmpty())
+				{
+					FSkeletalMeshImportData SerializeMeshData;
+					ThisLODModel.RawSkeletalMeshBulkData_DEPRECATED.LoadRawMesh(SerializeMeshData);
+					SaveLODImportedData(LODIndex, SerializeMeshData);
+				}
+				//Get the FRawSkeletalMeshBulkData to set the geo and skinning version
+				FRawSkeletalMeshBulkData& RawSkeletalMeshBulkData = GetMeshEditorData().GetLODImportedData(LODIndex);
+				RawSkeletalMeshBulkData.GeoImportVersion = ThisLODModel.RawSkeletalMeshBulkData_DEPRECATED.GeoImportVersion;
+				RawSkeletalMeshBulkData.SkinningImportVersion = ThisLODModel.RawSkeletalMeshBulkData_DEPRECATED.SkinningImportVersion;
+				//Empty the DEPRECATED member
+				FSkeletalMeshImportData EmptyMeshData;
+				ThisLODModel.RawSkeletalMeshBulkData_DEPRECATED.SaveRawMesh(EmptyMeshData);
+				ThisLODModel.RawSkeletalMeshBulkData_DEPRECATED.EmptyBulkData();
+			}
+			//Set the cache data into the LODModel
+			FRawSkeletalMeshBulkData& RawSkeletalMeshBulkData = GetMeshEditorData().GetLODImportedData(LODIndex);
+			ThisLODModel.bIsRawSkeletalMeshBulkDataEmpty = RawSkeletalMeshBulkData.IsEmpty();
+			ThisLODModel.bIsBuildDataAvailable = RawSkeletalMeshBulkData.IsBuildDataAvailable();
+			ThisLODModel.RawSkeletalMeshBulkDataID = RawSkeletalMeshBulkData.GetIdString();
+		}
+	}
+
+	if (GetLinkerCustomVersion(FEditorObjectVersion::GUID) < FEditorObjectVersion::SkeletalMeshBuildRefactor)
+	{
+		CreateUserSectionsDataForLegacyAssets();
+	}
+
+	PostLoadEnsureImportDataExist();
+
 #endif // #if WITH_EDITOR
 }
 
@@ -3281,10 +3399,7 @@ void USkeletalMesh::PostLoad()
 	BeginPostLoadInternal(Context);
 
 #if WITH_EDITOR
-	// Inline reduction is not thread-safe, prevent async build until this is fixed.
-	const bool bHasInlineReductions = SkeletalMeshImpl::HasInlineReductions(this);
-	const bool bForceSynchronous = Context.bIsPreSkeletalMeshBuildRefactor || bHasInlineReductions;
-	if (!bForceSynchronous && FSkeletalMeshCompilingManager::Get().IsAsyncCompilationAllowed(this))
+	if (FSkeletalMeshCompilingManager::Get().IsAsyncCompilationAllowed(this))
 	{
 		PrepareForAsyncCompilation();
 
@@ -3321,48 +3436,8 @@ void USkeletalMesh::ExecutePostLoadInternal(FSkeletalMeshPostLoadContext& Contex
 
 		UpdateGenerateUpToData();
 
-		if (GetLinkerCustomVersion(FEditorObjectVersion::GUID) < FEditorObjectVersion::SkeletalMeshMoveEditorSourceDataToPrivateAsset)
-		{
-			ReserveLODImportData(GetImportedModel()->LODModels.Num() - 1);
-			for (int32 LODIndex = 0; LODIndex < GetImportedModel()->LODModels.Num(); ++LODIndex)
-			{
-				FSkeletalMeshLODModel& ThisLODModel = GetImportedModel()->LODModels[LODIndex];
-				//We can have partial data if the asset was save after the split workflow implementation
-				//Use the deprecated member to retrieve this data
-				if (GetLinkerCustomVersion(FFortniteMainBranchObjectVersion::GUID) >= FFortniteMainBranchObjectVersion::NewSkeletalMeshImporterWorkflow)
-				{
-					if (!ThisLODModel.RawSkeletalMeshBulkData_DEPRECATED.IsEmpty())
-					{
-						FSkeletalMeshImportData SerializeMeshData;
-						ThisLODModel.RawSkeletalMeshBulkData_DEPRECATED.LoadRawMesh(SerializeMeshData);
-						SaveLODImportedData(LODIndex, SerializeMeshData);
-					}
-					//Get the FRawSkeletalMeshBulkData to set the geo and skinning version
-					FRawSkeletalMeshBulkData& RawSkeletalMeshBulkData = GetMeshEditorData().GetLODImportedData(LODIndex);
-					RawSkeletalMeshBulkData.GeoImportVersion = ThisLODModel.RawSkeletalMeshBulkData_DEPRECATED.GeoImportVersion;
-					RawSkeletalMeshBulkData.SkinningImportVersion = ThisLODModel.RawSkeletalMeshBulkData_DEPRECATED.SkinningImportVersion;
-					//Empty the DEPRECATED member
-					FSkeletalMeshImportData EmptyMeshData;
-					ThisLODModel.RawSkeletalMeshBulkData_DEPRECATED.SaveRawMesh(EmptyMeshData);
-					ThisLODModel.RawSkeletalMeshBulkData_DEPRECATED.EmptyBulkData();
-				}
-				//Set the cache data into the LODModel
-				FRawSkeletalMeshBulkData& RawSkeletalMeshBulkData = GetMeshEditorData().GetLODImportedData(LODIndex);
-				ThisLODModel.bIsRawSkeletalMeshBulkDataEmpty = RawSkeletalMeshBulkData.IsEmpty();
-				ThisLODModel.bIsBuildDataAvailable = RawSkeletalMeshBulkData.IsBuildDataAvailable();
-				ThisLODModel.RawSkeletalMeshBulkDataID = RawSkeletalMeshBulkData.GetIdString();
-			}
-		}
-
-		if (Context.bIsPreSkeletalMeshBuildRefactor)
-		{
-			CreateUserSectionsDataForLegacyAssets();
-		}
-
 		PostLoadValidateUserSectionData();
 
-		PostLoadEnsureImportDataExist();
-
 		PostLoadVerifyAndFixBadTangent();
 
 		if (GetResourceForRendering() == nullptr)
@@ -3431,11 +3506,9 @@ void USkeletalMesh::FinishPostLoadInternal(FSkeletalMeshPostLoadContext& Context
 
 #if WITH_EDITOR
 	ApplyFinishBuildInternalData(&Context);
-#else
-	// init morph targets here as in non-editor builds we don't use DDC or build the mesh. 
-	// should do this before InitResources, so that we clear invalid morph targets.
-	InitMorphTargets();
 #endif
+	// should do this before InitResources.
+	InitMorphTargets();
 
 	// initialize rendering resources
 	if (FApp::CanEverRender())
@@ -5205,9 +5278,10 @@ void USkeletalMesh::RemoveLODInfo(int32 Index)
 		{
 			GetMeshEditorDataObject()->RemoveLODImportedData(Index);
 		}
-		if (GetImportedModel() && GetImportedModel()->OriginalReductionSourceMeshData.IsValidIndex(Index))
+		
+		if (GetImportedModel()->InlineReductionCacheDatas.IsValidIndex(Index))
 		{
-			GetImportedModel()->OriginalReductionSourceMeshData.RemoveAt(Index);
+			GetImportedModel()->InlineReductionCacheDatas.RemoveAt(Index);
 		}
 #endif // WITH_EDITOR
 		LODInfoArray.RemoveAt(Index);
diff --git a/Engine/Source/Runtime/Engine/Private/SkeletalMeshLODRenderData.cpp b/Engine/Source/Runtime/Engine/Private/SkeletalMeshLODRenderData.cpp
--- a/Engine/Source/Runtime/Engine/Private/SkeletalMeshLODRenderData.cpp
+++ b/Engine/Source/Runtime/Engine/Private/SkeletalMeshLODRenderData.cpp
@@ -331,6 +331,19 @@ void FSkeletalMeshLODRenderData::InitResources(bool bNeedsVertexColors, int32 LO
 			int32 NumSrcDeltas = 0;
 			const FMorphTargetDelta* MorphDeltas = MorphTarget->GetMorphTargetDelta(LODIndex, NumSrcDeltas);
 
+			//Make sure the morphtarget data vertex indices fit the geometry
+			//If a missmatch happen, set the NumSrcDelta to 0 so the morph target is skipped
+			for (int32 DeltaIndex = 0; DeltaIndex < NumSrcDeltas; DeltaIndex++)
+			{
+				const auto& MorphDelta = MorphDeltas[DeltaIndex];
+				if (!VertexNeedsTangents.IsValidIndex(MorphDelta.SourceIdx))
+				{
+					NumSrcDeltas = 0;
+					UE_ASSET_LOG(LogSkeletalMesh, Warning, Owner, TEXT("Skipping morph target %s for LOD %d. The morph target data is incompatible with the mesh data"), *MorphTarget->GetName(), LODIndex);
+					break;
+				}
+			}
+
 			if (NumSrcDeltas == 0)
 			{
 				MaximumValues[0] = 0.0f;
diff --git a/Engine/Source/Runtime/Engine/Private/SkeletalMeshModel.cpp b/Engine/Source/Runtime/Engine/Private/SkeletalMeshModel.cpp
--- a/Engine/Source/Runtime/Engine/Private/SkeletalMeshModel.cpp
+++ b/Engine/Source/Runtime/Engine/Private/SkeletalMeshModel.cpp
@@ -7,6 +7,7 @@
 #include "Engine/SkeletalMesh.h"
 #include "Serialization/MemoryWriter.h"
 #include "EngineUtils.h"
+#include "UObject/UE5MainStreamObjectVersion.h"
 
 FSkeletalMeshModel::FSkeletalMeshModel()
 	: bGuidIsHash(false)	
@@ -26,6 +27,8 @@ void FSkeletalMeshModel::Serialize(FArchive& Ar, USkeletalMesh* Owner)
 	LODModels.Serialize(Ar, Owner);
 
 	Ar.UsingCustomVersion(FSkeletalMeshCustomVersion::GUID);
+	Ar.UsingCustomVersion(FUE5MainStreamObjectVersion::GUID);
+
 	// For old content without a GUID, generate one now from the data
 	if (Ar.IsLoading() && Ar.CustomVer(FSkeletalMeshCustomVersion::GUID) < FSkeletalMeshCustomVersion::SplitModelAndRenderData)
 	{
@@ -40,9 +43,65 @@ void FSkeletalMeshModel::Serialize(FArchive& Ar, USkeletalMesh* Owner)
 
 	if (!bIsEditorDataStripped)
 	{
-		if (Ar.IsSaving() || (Ar.CustomVer(FFortniteMainBranchObjectVersion::GUID) >= FFortniteMainBranchObjectVersion::AllowSkeletalMeshToReduceTheBaseLOD))
+		if (Ar.IsLoading()
+			&& (Ar.CustomVer(FFortniteMainBranchObjectVersion::GUID) >= FFortniteMainBranchObjectVersion::AllowSkeletalMeshToReduceTheBaseLOD)
+			&& (Ar.CustomVer(FUE5MainStreamObjectVersion::GUID) < FUE5MainStreamObjectVersion::ConvertReductionBaseSkeletalMeshBulkDataToInlineReductionCacheData))
+		{
+			FReductionBaseSkeletalMeshBulkData::Serialize(Ar, OriginalReductionSourceMeshData_DEPRECATED, Owner);
+
+			InlineReductionCacheDatas.Reserve(OriginalReductionSourceMeshData_DEPRECATED.Num());
+			for (int32 LodIndex = 0; LodIndex < OriginalReductionSourceMeshData_DEPRECATED.Num(); ++LodIndex)
+			{
+				if (!LODModels.IsValidIndex(LodIndex))
+				{
+					break;
+				}
+				bool bNeedToUseReductionData = true;
+				if (const FSkeletalMeshLODInfo* LODInfo = Owner->GetLODInfo(LodIndex))
+				{
+					if (!LODInfo->bHasBeenSimplified)
+					{
+						bNeedToUseReductionData = false;
+					}
+				}
+				if (bNeedToUseReductionData)
+				{
+					if (FReductionBaseSkeletalMeshBulkData* ReductionBaseSkeletalMeshBulkData = OriginalReductionSourceMeshData_DEPRECATED[LodIndex])
+					{
+						uint32 LodVertexCount = 0;
+						uint32 LodTriangleCount = 0;
+						ReductionBaseSkeletalMeshBulkData->GetGeometryInfo(LodVertexCount, LodTriangleCount, Owner);
+						InlineReductionCacheDatas.AddDefaulted_GetRef().SetCacheGeometryInfo(LodVertexCount, LodTriangleCount);
+					}
+				}
+				else
+				{
+					InlineReductionCacheDatas.AddDefaulted_GetRef().SetCacheGeometryInfo(LODModels[LodIndex]);
+				}
+			}
+		}
+
+		if (Ar.IsSaving() || (Ar.CustomVer(FUE5MainStreamObjectVersion::GUID) >= FUE5MainStreamObjectVersion::ConvertReductionBaseSkeletalMeshBulkDataToInlineReductionCacheData))
+		{
+			Ar << InlineReductionCacheDatas;
+		}
+
+		if (Ar.IsLoading())
 		{
-			FReductionBaseSkeletalMeshBulkData::Serialize(Ar, OriginalReductionSourceMeshData, Owner);
+			//Add any missing cache
+			if (InlineReductionCacheDatas.Num() < LODModels.Num())
+			{
+				InlineReductionCacheDatas.AddDefaulted(LODModels.Num() - InlineReductionCacheDatas.Num());
+				for (int32 LodIndex = InlineReductionCacheDatas.Num(); LodIndex < LODModels.Num(); ++LodIndex)
+				{
+					InlineReductionCacheDatas[LodIndex].SetCacheGeometryInfo(LODModels[LodIndex]);
+				}
+			}
+			else if (InlineReductionCacheDatas.Num() > LODModels.Num())
+			{
+				//If we have too much entry simply shrink the array to valid LODModel size
+				InlineReductionCacheDatas.SetNum(LODModels.Num());
+			}
 		}
 	}
 }
@@ -126,12 +185,12 @@ void FSkeletalMeshModel::GetResourceSizeEx(FResourceSizeEx& CumulativeResourceSi
 
 void FSkeletalMeshModel::EmptyOriginalReductionSourceMeshData()
 {
-	for (FReductionBaseSkeletalMeshBulkData* ReductionData : OriginalReductionSourceMeshData)
+	for (FReductionBaseSkeletalMeshBulkData* ReductionData : OriginalReductionSourceMeshData_DEPRECATED)
 	{
 		ReductionData->EmptyBulkData();
 		delete ReductionData;
 	}
-	OriginalReductionSourceMeshData.Empty();
+	OriginalReductionSourceMeshData_DEPRECATED.Empty();
 }
 
 #endif // WITH_EDITOR
\ No newline at end of file
diff --git a/Engine/Source/Runtime/Engine/Private/SkeletalMeshRenderData.cpp b/Engine/Source/Runtime/Engine/Private/SkeletalMeshRenderData.cpp
--- a/Engine/Source/Runtime/Engine/Private/SkeletalMeshRenderData.cpp
+++ b/Engine/Source/Runtime/Engine/Private/SkeletalMeshRenderData.cpp
@@ -345,10 +345,21 @@ void FSkeletalMeshRenderData::Cache(const ITargetPlatform* TargetPlatform, USkel
 	{
 		COOK_STAT(auto Timer = SkeletalMeshCookStats::UsageStats.TimeSyncWork());
 		int32 T0 = FPlatformTime::Cycles();
-		DerivedDataKey = BuildSkeletalMeshDerivedDataKey(TargetPlatform, Owner);
 
+		//When we import a skeletalmesh, in some cases the asset is not yet built, and the usersectiondata and the inline cache are not set
+		//until the initial build. This is due to the section count which is establish by the initial build of the import data. The section count
+		//is part of the key because users can change section settings(see UserSectionData). So when we do a initial build we do not compute yet
+		//the key and force the build code path, the key will be compute after the build and the DDC data will be store with the computed key.
+		const bool bAllowDdcFetch = Owner->IsInitialBuildDone();
+		if (bAllowDdcFetch)
+		{
+			DerivedDataKey = BuildSkeletalMeshDerivedDataKey(TargetPlatform, Owner);
+		}
+		
+		//If we have an initial build, the ddc key will be computed only after the build. Some structure are missing until we first build the asset to get the drived data key
+		
 		TArray64<uint8> DerivedData;
-		if(DDCUtils64Bit::GetSynchronous(DerivedDataKey, Owner, DerivedData))
+		if(bAllowDdcFetch && DDCUtils64Bit::GetSynchronous(DerivedDataKey, Owner, DerivedData))
 		{
 			COOK_STAT(Timer.AddHit(DerivedData.Num()));
 			
@@ -514,12 +525,21 @@ void FSkeletalMeshRenderData::Cache(const ITargetPlatform* TargetPlatform, USkel
 			//Recompute the derived data key in case there was some data correction during the build process, this make sure the DDC key is always representing the correct build result.
 			//There should never be correction of the data during the build, the data has to be corrected in the post load before calling this function.
 			FString BuiltDerivedDataKey = BuildSkeletalMeshDerivedDataKey(TargetPlatform, Owner);
-			if(BuiltDerivedDataKey != DerivedDataKey)
+			//Only compare keys if the ddc fetch was allowed
+			if (bAllowDdcFetch)
+			{
+				if (BuiltDerivedDataKey != DerivedDataKey)
+				{
+					//If we are in this case we should resave the asset so the source data will be the same and we can use this DDC. Reduction can change the number of sections and the user section data is in the DDC key.
+					//So if we change the reduction algorithm, its possible we fall in this situation.
+					//We save the real data key which force the asset to always rebuild when the editor is loading it until the user save it
+					UE_LOG(LogSkeletalMesh, Log, TEXT("Skeletal mesh [%s]: The derived data key is different after the build. Resave the asset to avoid rebuilding it everytime the editor load it."), *Owner->GetPathName());
+				}
+			}
+			else
 			{
-				//If we are in this case we should resave the asset so the source data will be the same and we can use this DDC. Reduction can change the number of sections and the user section data is in the DDC key.
-				//So if we change the reduction algorithm, its possible we fall in this situation.
-				//We save the real data key which force the asset to always rebuild when the editor is loading it until the user save it
-				UE_LOG(LogSkeletalMesh, Log, TEXT("Skeletal mesh [%s]: The derived data key is different after the build. Resave the asset to avoid rebuilding it everytime the editor load it."), *Owner->GetPathName());
+				//After the initial build we set the key to the built one
+				DerivedDataKey = BuiltDerivedDataKey;
 			}
 
 			//Store the data using the built key to avoid DDC corruption
diff --git a/Engine/Source/Runtime/Engine/Public/Rendering/SkeletalMeshLODImporterData.h b/Engine/Source/Runtime/Engine/Public/Rendering/SkeletalMeshLODImporterData.h
--- a/Engine/Source/Runtime/Engine/Public/Rendering/SkeletalMeshLODImporterData.h
+++ b/Engine/Source/Runtime/Engine/Public/Rendering/SkeletalMeshLODImporterData.h
@@ -433,7 +433,7 @@ private:
 };
 
 /**
-* Bulk data storage for raw ImportModel.
+* Bulk data storage for raw ImportModel. This structure is deprecated, we now only store the original vertex and triangle count, see FInlineReductionCacheData.
 */
 struct FReductionBaseSkeletalMeshBulkData
 {
@@ -480,6 +480,29 @@ public:
 	FORCEINLINE bool IsEmpty() const { return BulkData.GetBulkDataSize() == 0; }
 };
 
+struct FInlineReductionCacheData
+{
+	uint32 CacheLODVertexCount = MAX_uint32;
+	uint32 CacheLODTriCount = MAX_uint32;
+
+	/*
+	 * Caching those value since this is a slow operation to load the bulk data to retrieve the original geometry information
+	 */
+	ENGINE_API void SetCacheGeometryInfo(const FSkeletalMeshLODModel& SourceLODModel);
+
+	ENGINE_API void SetCacheGeometryInfo(uint32 LODVertexCount, uint32 LODTriCount);
+
+	/** Return the cache count of vertices and triangles. */
+	ENGINE_API void GetCacheGeometryInfo(uint32& LODVertexCount, uint32& LODTriCount) const;
+};
+
+FORCEINLINE FArchive& operator<<(FArchive& Ar, FInlineReductionCacheData& InlineReductionCacheData)
+{
+	Ar << InlineReductionCacheData.CacheLODVertexCount;
+	Ar << InlineReductionCacheData.CacheLODTriCount;
+	return Ar;
+}
+
 /**
 * Bulk data storage for raw meshes.
 */
diff --git a/Engine/Source/Runtime/Engine/Public/Rendering/SkeletalMeshModel.h b/Engine/Source/Runtime/Engine/Public/Rendering/SkeletalMeshModel.h
--- a/Engine/Source/Runtime/Engine/Public/Rendering/SkeletalMeshModel.h
+++ b/Engine/Source/Runtime/Engine/Public/Rendering/SkeletalMeshModel.h
@@ -47,8 +47,14 @@ public:
 	FString GetLODModelIdString() const;
 
 	ENGINE_API void EmptyOriginalReductionSourceMeshData();
+	
 	/* When user reduce an imported LOD with himself (BaseLOD == TargetLOD), we need to store some imported model data so we can reduce again from the same data.*/
-	TArray<FReductionBaseSkeletalMeshBulkData*> OriginalReductionSourceMeshData;
+	/* We do not need to store such a data, since we can use the USkeletalMesh::MeshEditorDataObject which have the imported data */
+	TArray<FReductionBaseSkeletalMeshBulkData*> OriginalReductionSourceMeshData_DEPRECATED;
+
+	/* When user reduce an imported LOD with himself (BaseLOD == TargetLOD), we need to store the geometry count (vertex count and triangle count).
+	We use this data to query if the reduction is active when the reduction criterion is an absolute vertex/triangle number*/
+	TArray<FInlineReductionCacheData> InlineReductionCacheDatas;
 #endif
 
 	/** Serialize to/from the specified archive.. */
