Fix alternate skinning issue

diff --git a/Engine/Source/Developer/SkeletalMeshUtilitiesCommon/Private/LODUtilities.cpp b/Engine/Source/Developer/SkeletalMeshUtilitiesCommon/Private/LODUtilities.cpp
--- a/Engine/Source/Developer/SkeletalMeshUtilitiesCommon/Private/LODUtilities.cpp
+++ b/Engine/Source/Developer/SkeletalMeshUtilitiesCommon/Private/LODUtilities.cpp
@@ -1858,6 +1858,19 @@ bool FLODUtilities::UpdateAlternateSkinWeights(FSkeletalMeshLODModel& LODModelDe
 	TSortedMap<uint32, VertexMatchNameSpace::FVertexMatchResult> VertexIndexSrcToVertexIndexDestMatches;
 	VertexIndexSrcToVertexIndexDestMatches.Reserve(VertexNumberSrc);
 	TArray<uint32> VertexIndexToMatchWithPositions;
+
+	auto FindWedgeIndexesUsingVertexIndex = [](const FSkeletalMeshImportData& ImportData, const int32 VertexIndex, TArray<int32>& OutWedgeIndexes)
+	{
+		for (int32 WedgeIndex = 0; WedgeIndex < ImportData.Wedges.Num(); ++WedgeIndex)
+		{
+			const SkeletalMeshImportData::FVertex& Wedge = ImportData.Wedges[WedgeIndex];
+			if (Wedge.VertexIndex == VertexIndex)
+			{
+				OutWedgeIndexes.Add(WedgeIndex);
+			}
+		}
+	};
+
 	// Match all source vertex with destination vertex
 	for (int32 VertexIndexSrc = 0; VertexIndexSrc < PointNumberSrc; ++VertexIndexSrc)
 	{
@@ -1875,10 +1888,46 @@ bool FLODUtilities::UpdateAlternateSkinWeights(FSkeletalMeshLODModel& LODModelDe
 		{
 			//We have a direct match
 			VertexMatchNameSpace::FVertexMatchResult& VertexMatchDest = VertexIndexSrcToVertexIndexDestMatches.Add(VertexIndexSrc);
+
+			TArray<int32> SrcWedgeIndexes;
+			FindWedgeIndexesUsingVertexIndex(ImportDataSrc, VertexIndexSrc, SrcWedgeIndexes);
+
+			//Check if we have a point that is perfectly matching (position, UV, material and vertex color). Because normals and tangent are on the triangles we do not test those.
 			for (int32 MatchDestinationIndex = 0; MatchDestinationIndex < SimilarDestinationVertex.Num(); ++MatchDestinationIndex)
 			{
-				VertexMatchDest.VertexIndexes.Add(SimilarDestinationVertex[MatchDestinationIndex]);
-				VertexMatchDest.Ratios.Add(1.0f);
+				int32 VertexIndexDest = SimilarDestinationVertex[MatchDestinationIndex];
+				TArray<int32> DestWedgeIndexes;
+				FindWedgeIndexesUsingVertexIndex(ImportDataDest, VertexIndexDest, DestWedgeIndexes);
+				for (int32 IndexDest = 0; IndexDest < DestWedgeIndexes.Num(); ++IndexDest)
+				{
+					int32 DestWedgeIndex = DestWedgeIndexes[IndexDest];
+					const SkeletalMeshImportData::FVertex& WedgeDest = ImportDataDest.Wedges[DestWedgeIndex];
+					for (int32 IndexSrc = 0; IndexSrc < SrcWedgeIndexes.Num(); ++IndexSrc)
+					{
+						int32 SrcWedgeIndex = SrcWedgeIndexes[IndexSrc];
+						const SkeletalMeshImportData::FVertex& WedgeSrc = ImportDataSrc.Wedges[SrcWedgeIndex];
+						//Wedge == operator test: material, vertex color and UVs
+						if (WedgeDest == WedgeSrc)
+						{
+							VertexMatchDest.VertexIndexes.Add(SimilarDestinationVertex[MatchDestinationIndex]);
+							VertexMatchDest.Ratios.Add(1.0f);
+							break;
+						}
+					}
+					if (VertexMatchDest.VertexIndexes.Num() > 0)
+					{
+						break;
+					}
+				}
+			}
+			//If there is no direct match, simply put everything
+			if (VertexMatchDest.VertexIndexes.Num() == 0)
+			{
+				for (int32 MatchDestinationIndex = 0; MatchDestinationIndex < SimilarDestinationVertex.Num(); ++MatchDestinationIndex)
+				{
+					VertexMatchDest.VertexIndexes.Add(SimilarDestinationVertex[MatchDestinationIndex]);
+					VertexMatchDest.Ratios.Add(1.0f);
+				}
 			}
 		}
 	}
diff --git a/Engine/Source/Runtime/Engine/Private/SkeletalMeshRenderData.cpp b/Engine/Source/Runtime/Engine/Private/SkeletalMeshRenderData.cpp
--- a/Engine/Source/Runtime/Engine/Private/SkeletalMeshRenderData.cpp
+++ b/Engine/Source/Runtime/Engine/Private/SkeletalMeshRenderData.cpp
@@ -186,7 +186,7 @@ static void SerializeLODInfoForDDC(USkeletalMesh* SkeletalMesh, FString& KeySuff
 // differences, etc.) replace the version GUID below with a new one.
 // In case of merge conflicts with DDC versions, you MUST generate a new GUID
 // and set this new GUID as the version.
-#define SKELETALMESH_DERIVEDDATA_VER TEXT("F1115E3EF0EE404692FABA5C90B76988")
+#define SKELETALMESH_DERIVEDDATA_VER TEXT("0F778E3964824C509A348A7D390C13D8")
 
 const FString& GetSkeletalMeshDerivedDataVersion()
 {
diff --git a/Engine/Source/Runtime/Engine/Public/Rendering/SkeletalMeshLODImporterData.h b/Engine/Source/Runtime/Engine/Public/Rendering/SkeletalMeshLODImporterData.h
--- a/Engine/Source/Runtime/Engine/Public/Rendering/SkeletalMeshLODImporterData.h
+++ b/Engine/Source/Runtime/Engine/Public/Rendering/SkeletalMeshLODImporterData.h
@@ -269,7 +269,7 @@ namespace SkeletalMeshImportData
 			bool bUVsEqual = true;
 			for (uint32 UVIdx = 0; UVIdx < MAX_TEXCOORDS; ++UVIdx)
 			{
-				if (UVs[UVIdx] != Other.UVs[UVIdx])
+				if (!UVs[UVIdx].Equals(Other.UVs[UVIdx], SMALL_NUMBER))
 				{
 					bUVsEqual = false;
 					break;
