Fix a lock issue in CopyStructure, bulkdata are not playing well with thread so I prefer not using them for such small amount of editor data.
Removing them remove the following debug assert in LinkerLoad.h that happen when we copy FSkeletalMeshLODModel structure.

checkSlow(FPlatformTLS::GetCurrentThreadId() == OwnerThread);

It should also prevent a crash with file handle since we cannot lock two different bulk data in separate thread if the two bulk data are store in the same package since they will share a file handle

#rnx

diff --git a/Engine/Plugins/Experimental/AlembicImporter/Source/AlembicLibrary/Private/AbcImporter.cpp b/Engine/Plugins/Experimental/AlembicImporter/Source/AlembicLibrary/Private/AbcImporter.cpp
--- a/Engine/Plugins/Experimental/AlembicImporter/Source/AlembicLibrary/Private/AbcImporter.cpp
+++ b/Engine/Plugins/Experimental/AlembicImporter/Source/AlembicLibrary/Private/AbcImporter.cpp
@@ -1621,7 +1621,9 @@ bool FAbcImporter::BuildSkeletalMesh( FSkeletalMeshLODModel& LODModel, const FRe
 	LODModel.NumVertices = 0;
 	LODModel.IndexBuffer.Empty();
 
-	TArray<uint32> RawPointIndices;
+	TArray<uint32>& RawPointIndices = LODModel.GetRawPointIndices();
+	RawPointIndices.Reset();
+
 	TArray< TArray<uint32> > VertexIndexRemap;
 	VertexIndexRemap.Empty(MeshSections.Num());
 
@@ -1720,16 +1722,6 @@ bool FAbcImporter::BuildSkeletalMesh( FSkeletalMeshLODModel& LODModel, const FRe
 	// Only using bone zero
 	LODModel.ActiveBoneIndices.Add(0);
 
-	// Copy raw point indices to LOD model.
-	LODModel.RawPointIndices.RemoveBulkData();
-	if (RawPointIndices.Num())
-	{
-		LODModel.RawPointIndices.Lock(LOCK_READ_WRITE);
-		void* Dest = LODModel.RawPointIndices.Realloc(RawPointIndices.Num());
-		FMemory::Memcpy(Dest, RawPointIndices.GetData(), LODModel.RawPointIndices.GetBulkDataSize());
-		LODModel.RawPointIndices.Unlock();
-	}
-
 	// Finish building the sections.
 	for (int32 SectionIndex = 0; SectionIndex < LODModel.Sections.Num(); SectionIndex++)
 	{
diff --git a/Engine/Source/Developer/MeshUtilities/Private/MeshUtilities.cpp b/Engine/Source/Developer/MeshUtilities/Private/MeshUtilities.cpp
--- a/Engine/Source/Developer/MeshUtilities/Private/MeshUtilities.cpp
+++ b/Engine/Source/Developer/MeshUtilities/Private/MeshUtilities.cpp
@@ -749,8 +749,9 @@ void FMeshUtilities::BuildSkeletalModelFromChunks(FSkeletalMeshLODModel& LODMode
 	// Keep track of index mapping to chunk vertex offsets
 	TArray< TArray<uint32> > VertexIndexRemap;
 	VertexIndexRemap.Empty(LODModel.Sections.Num());
+	TArray<uint32>& RawPointIndices = LODModel.GetRawPointIndices();
+	RawPointIndices.Reset();
 	// Pack the chunk vertices into a single vertex buffer.
-	TArray<uint32> RawPointIndices;
 	LODModel.NumVertices = 0;
 
 	int32 PrevMaterialIndex = -1;
@@ -882,16 +883,6 @@ void FMeshUtilities::BuildSkeletalModelFromChunks(FSkeletalMeshLODModel& LODMode
 			);
 	}
 
-	// Copy raw point indices to LOD model.
-	LODModel.RawPointIndices.RemoveBulkData();
-	if (RawPointIndices.Num())
-	{
-		LODModel.RawPointIndices.Lock(LOCK_READ_WRITE);
-		void* Dest = LODModel.RawPointIndices.Realloc(RawPointIndices.Num());
-		FMemory::Memcpy(Dest, RawPointIndices.GetData(), LODModel.RawPointIndices.GetBulkDataSize());
-		LODModel.RawPointIndices.Unlock();
-	}
-
 	// Finish building the sections.
 	for (int32 SectionIndex = 0; SectionIndex < LODModel.Sections.Num(); SectionIndex++)
 	{
diff --git a/Engine/Source/Developer/SkeletalMeshUtilitiesCommon/Private/LODUtilities.cpp b/Engine/Source/Developer/SkeletalMeshUtilitiesCommon/Private/LODUtilities.cpp
--- a/Engine/Source/Developer/SkeletalMeshUtilitiesCommon/Private/LODUtilities.cpp
+++ b/Engine/Source/Developer/SkeletalMeshUtilitiesCommon/Private/LODUtilities.cpp
@@ -2457,7 +2457,7 @@ class FAsyncImportMorphTargetWork : public FNonAbandonableTask
 {
 public:
 	FAsyncImportMorphTargetWork(FSkeletalMeshLODModel* InLODModel, const FReferenceSkeleton& InRefSkeleton, const FSkeletalMeshImportData& InBaseImportData, TArray<FVector3f>&& InMorphLODPoints,
-		TArray< FMorphTargetDelta >& InMorphDeltas, TArray<uint32>& InBaseIndexData, TArray< uint32 >& InBaseWedgePointIndices,
+		TArray< FMorphTargetDelta >& InMorphDeltas, TArray<uint32>& InBaseIndexData, const TArray< uint32 >& InBaseWedgePointIndices,
 		TMap<uint32, uint32>& InWedgePointToVertexIndexMap, const FOverlappingCorners& InOverlappingCorners,
 		const TSet<uint32> InModifiedPoints, const TMultiMap< int32, int32 >& InWedgeToFaces, const FMeshDataBundle& InMeshDataBundle, const TArray<FVector3f>& InTangentZ,
 		bool InShouldImportNormals, bool InShouldImportTangents, bool InbUseMikkTSpace, const FOverlappingThresholds InThresholds)
@@ -2688,7 +2688,7 @@ private:
 
 	TArray< FMorphTargetDelta >& MorphTargetDeltas;
 	TArray< uint32 >& BaseIndexData;
-	TArray< uint32 >& BaseWedgePointIndices;
+	const TArray< uint32 >& BaseWedgePointIndices;
 	TMap<uint32, uint32>& WedgePointToVertexIndexMap;
 
 	const FOverlappingCorners& OverlappingCorners;
@@ -2732,15 +2732,6 @@ void FLODUtilities::BuildMorphTargets(USkeletalMesh* BaseSkelMesh, FSkeletalMesh
 	TArray<FVector3f> TangentZ;
 	MeshUtilities.CalculateNormals(MeshDataBundle.Vertices, MeshDataBundle.Indices, MeshDataBundle.UVs, MeshDataBundle.SmoothingGroups, TangentOptions, TangentZ);
 
-	TArray< uint32 > BaseWedgePointIndices;
-	if (BaseLODModel.RawPointIndices.GetBulkDataSize())
-	{
-		BaseWedgePointIndices.Empty(BaseLODModel.RawPointIndices.GetElementCount());
-		BaseWedgePointIndices.AddUninitialized(BaseLODModel.RawPointIndices.GetElementCount());
-		FMemory::Memcpy(BaseWedgePointIndices.GetData(), BaseLODModel.RawPointIndices.Lock(LOCK_READ_ONLY), BaseLODModel.RawPointIndices.GetBulkDataSize());
-		BaseLODModel.RawPointIndices.Unlock();
-	}
-
 	TArray<uint32> BaseIndexData = BaseLODModel.IndexBuffer;
 
 	TMap<uint32, uint32> WedgePointToVertexIndexMap;
@@ -2861,7 +2852,7 @@ void FLODUtilities::BuildMorphTargets(USkeletalMesh* BaseSkelMesh, FSkeletalMesh
 			TArray< FMorphTargetDelta >* Deltas = Results[NewMorphDeltasIdx];
 
 			FAsyncTask<FAsyncImportMorphTargetWork>* NewWork = new FAsyncTask<FAsyncImportMorphTargetWork>(&BaseLODModel, BaseSkelMesh->GetRefSkeleton(), BaseImportData,
-				MoveTemp(ShapeImportData.Points), *Deltas, BaseIndexData, BaseWedgePointIndices, WedgePointToVertexIndexMap, OverlappingVertices, MoveTemp(ModifiedPoints), WedgeToFaces, MeshDataBundle, TangentZ,
+				MoveTemp(ShapeImportData.Points), *Deltas, BaseIndexData, BaseLODModel.GetRawPointIndices(), WedgePointToVertexIndexMap, OverlappingVertices, MoveTemp(ModifiedPoints), WedgeToFaces, MeshDataBundle, TangentZ,
 				ShouldImportNormals, ShouldImportTangents, bUseMikkTSpace, Thresholds);
 			PendingWork.Add(NewWork);
 
@@ -3314,17 +3305,8 @@ bool FLODUtilities::StripLODGeometry(USkeletalMesh* SkeletalMesh, const int32 LO
 			return true;
 		};
 
-		TArray< uint32 > SoftVertexIndexToImportDataPointIndex;
-		if (LODModel.RawPointIndices.GetBulkDataSize())
-		{
-			SoftVertexIndexToImportDataPointIndex.Empty(LODModel.RawPointIndices.GetElementCount());
-			SoftVertexIndexToImportDataPointIndex.AddUninitialized(LODModel.RawPointIndices.GetElementCount());
-			FMemory::Memcpy(SoftVertexIndexToImportDataPointIndex.GetData(), LODModel.RawPointIndices.Lock(LOCK_READ_ONLY), LODModel.RawPointIndices.GetBulkDataSize());
-			LODModel.RawPointIndices.Unlock();
-		}
+		const TArray< uint32 >& SoftVertexIndexToImportDataPointIndex = LODModel.GetRawPointIndices();
 
-		
-		
 		TMap<uint64, TArray<int32>> OptimizedFaceFinder;
 		
 		auto GetMatchFaceIndex = [&OptimizedFaceFinder, &ImportedData](const int32 FaceVertexA, const int32 FaceVertexB, int32 FaceVertexC)->int32
diff --git a/Engine/Source/Runtime/Core/Public/UObject/UE5ReleaseStreamObjectVersion.h b/Engine/Source/Runtime/Core/Public/UObject/UE5ReleaseStreamObjectVersion.h
--- a/Engine/Source/Runtime/Core/Public/UObject/UE5ReleaseStreamObjectVersion.h
+++ b/Engine/Source/Runtime/Core/Public/UObject/UE5ReleaseStreamObjectVersion.h
@@ -75,6 +75,9 @@ struct CORE_API FUE5ReleaseStreamObjectVersion
 		// Add support for actor folder objects
 		AddLevelActorFolders,
 
+		// Remove FSkeletalMeshLODModel bulk datas
+		RemoveSkeletalMeshLODModelBulkDatas,
+		
 		// -----<new versions can be added above this line>-------------------------------------------------
 		VersionPlusOne,
 		LatestVersion = VersionPlusOne - 1
diff --git a/Engine/Source/Runtime/Engine/Private/SkeletalMesh.cpp b/Engine/Source/Runtime/Engine/Private/SkeletalMesh.cpp
--- a/Engine/Source/Runtime/Engine/Private/SkeletalMesh.cpp
+++ b/Engine/Source/Runtime/Engine/Private/SkeletalMesh.cpp
@@ -3507,20 +3507,28 @@ void USkeletalMesh::BeginPostLoadInternal(FSkeletalMeshPostLoadContext& Context)
 				//Use the deprecated member to retrieve this data
 				if (GetLinkerCustomVersion(FFortniteMainBranchObjectVersion::GUID) >= FFortniteMainBranchObjectVersion::NewSkeletalMeshImporterWorkflow)
 				{
-					if (!ThisLODModel.RawSkeletalMeshBulkData_DEPRECATED.IsEmpty())
+					//Get the deprecated data 
+					FRawSkeletalMeshBulkData* RawSkeletalMeshBulkData_DEPRECATED = nullptr;
+					{
+						PRAGMA_DISABLE_DEPRECATION_WARNINGS
+						RawSkeletalMeshBulkData_DEPRECATED = &ThisLODModel.GetRawSkeletalMeshBulkData_DEPRECATED();
+						PRAGMA_ENABLE_DEPRECATION_WARNINGS
+					}
+					check(RawSkeletalMeshBulkData_DEPRECATED);
+					if (!RawSkeletalMeshBulkData_DEPRECATED->IsEmpty())
 					{
 						FSkeletalMeshImportData SerializeMeshData;
-						ThisLODModel.RawSkeletalMeshBulkData_DEPRECATED.LoadRawMesh(SerializeMeshData);
+						RawSkeletalMeshBulkData_DEPRECATED->LoadRawMesh(SerializeMeshData);
 						SaveLODImportedData(LODIndex, SerializeMeshData);
 					}
 					//Get the FRawSkeletalMeshBulkData to set the geo and skinning version
 					FRawSkeletalMeshBulkData& RawSkeletalMeshBulkData = GetMeshEditorData().GetLODImportedData(LODIndex);
-					RawSkeletalMeshBulkData.GeoImportVersion = ThisLODModel.RawSkeletalMeshBulkData_DEPRECATED.GeoImportVersion;
-					RawSkeletalMeshBulkData.SkinningImportVersion = ThisLODModel.RawSkeletalMeshBulkData_DEPRECATED.SkinningImportVersion;
+					RawSkeletalMeshBulkData.GeoImportVersion = RawSkeletalMeshBulkData_DEPRECATED->GeoImportVersion;
+					RawSkeletalMeshBulkData.SkinningImportVersion = RawSkeletalMeshBulkData_DEPRECATED->SkinningImportVersion;
 					//Empty the DEPRECATED member
 					FSkeletalMeshImportData EmptyMeshData;
-					ThisLODModel.RawSkeletalMeshBulkData_DEPRECATED.SaveRawMesh(EmptyMeshData);
-					ThisLODModel.RawSkeletalMeshBulkData_DEPRECATED.EmptyBulkData();
+					RawSkeletalMeshBulkData_DEPRECATED->SaveRawMesh(EmptyMeshData);
+					RawSkeletalMeshBulkData_DEPRECATED->EmptyBulkData();
 				}
 				//Set the cache data into the LODModel
 				FRawSkeletalMeshBulkData& RawSkeletalMeshBulkData = GetMeshEditorData().GetLODImportedData(LODIndex);
diff --git a/Engine/Source/Runtime/Engine/Private/SkeletalMeshLODModel.cpp b/Engine/Source/Runtime/Engine/Private/SkeletalMeshLODModel.cpp
--- a/Engine/Source/Runtime/Engine/Private/SkeletalMeshLODModel.cpp
+++ b/Engine/Source/Runtime/Engine/Private/SkeletalMeshLODModel.cpp
@@ -740,6 +740,7 @@ void FSkeletalMeshLODModel::Serialize(FArchive& Ar, UObject* Owner, int32 Idx)
 	Ar.UsingCustomVersion(FFortniteMainBranchObjectVersion::GUID);
 	Ar.UsingCustomVersion(FEditorObjectVersion::GUID);
 	Ar.UsingCustomVersion(FUE5MainStreamObjectVersion::GUID);
+	Ar.UsingCustomVersion(FUE5ReleaseStreamObjectVersion::GUID);
 
 	if (StripFlags.IsDataStrippedForServer())
 	{
@@ -851,7 +852,27 @@ void FSkeletalMeshLODModel::Serialize(FArchive& Ar, UObject* Owner, int32 Idx)
 
 	if (!StripFlags.IsEditorDataStripped())
 	{
-		RawPointIndices.Serialize(Ar, Owner);
+		if (Ar.IsLoading() && Ar.CustomVer(FUE5ReleaseStreamObjectVersion::GUID) < FUE5ReleaseStreamObjectVersion::RemoveSkeletalMeshLODModelBulkDatas)
+		{
+			RawPointIndices_DEPRECATED.Serialize(Ar, Owner);
+			if (RawPointIndices_DEPRECATED.GetBulkDataSize())
+			{
+				if (RawPointIndices_DEPRECATED.IsAsyncLoadingComplete() && !RawPointIndices_DEPRECATED.IsBulkDataLoaded())
+				{
+					RawPointIndices_DEPRECATED.LoadBulkDataWithFileReader();
+				}
+			
+				RawPointIndices2.Empty(RawPointIndices_DEPRECATED.GetElementCount());
+				RawPointIndices2.AddUninitialized(RawPointIndices_DEPRECATED.GetElementCount());
+				FMemory::Memcpy(RawPointIndices2.GetData(), RawPointIndices_DEPRECATED.Lock(LOCK_READ_ONLY), RawPointIndices_DEPRECATED.GetBulkDataSize());
+				RawPointIndices_DEPRECATED.Unlock();
+			}
+			RawPointIndices_DEPRECATED.RemoveBulkData();
+		}
+		else
+		{
+			Ar << RawPointIndices2;
+		}
 		if (Ar.IsLoading()
 			&& (Ar.CustomVer(FFortniteMainBranchObjectVersion::GUID) >= FFortniteMainBranchObjectVersion::NewSkeletalMeshImporterWorkflow)
 			&& (Ar.CustomVer(FEditorObjectVersion::GUID) < FEditorObjectVersion::SkeletalMeshMoveEditorSourceDataToPrivateAsset))
@@ -1080,8 +1101,7 @@ void FSkeletalMeshLODModel::GetResourceSizeEx(FResourceSizeEx& CumulativeResourc
 	CumulativeResourceSize.AddUnknownMemoryBytes(RequiredBones.GetAllocatedSize());
 	CumulativeResourceSize.AddUnknownMemoryBytes(IndexBuffer.GetAllocatedSize());
 
-	CumulativeResourceSize.AddUnknownMemoryBytes(RawPointIndices.GetBulkDataSize());
-	CumulativeResourceSize.AddUnknownMemoryBytes(LegacyRawPointIndices.GetBulkDataSize());
+	CumulativeResourceSize.AddUnknownMemoryBytes(RawPointIndices2.GetAllocatedSize());
 	CumulativeResourceSize.AddUnknownMemoryBytes(MeshToImportVertexMap.GetAllocatedSize());
 }
 
@@ -1340,19 +1360,7 @@ void FSkeletalMeshLODModel::CopyStructure(FSkeletalMeshLODModel* Destination, co
 
 	FCriticalSection* DestinationBulkDataReadMutex = Destination->BulkDataReadMutex;
 
-	//Empty the Destination BulkData to avoid leaks
-	Destination->RawPointIndices.RemoveBulkData();
-	Destination->LegacyRawPointIndices.RemoveBulkData();
-	Destination->RawSkeletalMeshBulkData_DEPRECATED.EmptyBulkData();
-
-	// Bulk data arrays need to be locked before a copy can be made.
-	Source->RawPointIndices.LockReadOnly();
-	Source->LegacyRawPointIndices.LockReadOnly();
-	Source->RawSkeletalMeshBulkData_DEPRECATED.GetBulkData().LockReadOnly();
 	*Destination = *Source;
-	Source->RawSkeletalMeshBulkData_DEPRECATED.GetBulkData().Unlock();
-	Source->RawPointIndices.Unlock();
-	Source->LegacyRawPointIndices.Unlock();
 
 	//Make sure the mutex of the copy is set back to the original destination mutex, we can recycle the pointer.
 	Destination->BulkDataReadMutex = DestinationBulkDataReadMutex;
diff --git a/Engine/Source/Runtime/Engine/Private/SkeletalMeshRenderData.cpp b/Engine/Source/Runtime/Engine/Private/SkeletalMeshRenderData.cpp
--- a/Engine/Source/Runtime/Engine/Private/SkeletalMeshRenderData.cpp
+++ b/Engine/Source/Runtime/Engine/Private/SkeletalMeshRenderData.cpp
@@ -332,14 +332,8 @@ void FSkeletalMeshRenderData::Cache(const ITargetPlatform* TargetPlatform, USkel
 		Ar << LODModel->RequiredBones;
 		Ar << LODModel->MeshToImportVertexMap;
 		Ar << LODModel->MaxImportVertex;
-
-		// Unless an async loading query is already in flight, we want to load directly from disk to avoid
-		// accessing the linker which is not thread-safe.
-		if (LODModel->RawPointIndices.IsAsyncLoadingComplete() && !LODModel->RawPointIndices.IsBulkDataLoaded())
-		{
-			LODModel->RawPointIndices.LoadBulkDataWithFileReader();
-		}
-		LODModel->RawPointIndices.Serialize(Ar, Owner);
+		TArray<uint32>& RawPointIndices = LODModel->GetRawPointIndices();
+		Ar << RawPointIndices;
 	};
 
 	{
diff --git a/Engine/Source/Runtime/Engine/Public/Rendering/SkeletalMeshLODModel.h b/Engine/Source/Runtime/Engine/Public/Rendering/SkeletalMeshLODModel.h
--- a/Engine/Source/Runtime/Engine/Public/Rendering/SkeletalMeshLODModel.h
+++ b/Engine/Source/Runtime/Engine/Public/Rendering/SkeletalMeshLODModel.h
@@ -386,12 +386,24 @@ public:
 	/** The max index in MeshToImportVertexMap, ie. the number of imported (raw) verts. */
 	int32						MaxImportVertex;
 
-	/** Editor only data: array of the original point (wedge) indices for each of the vertices in a FSkeletalMeshLODModel */
-	FIntBulkData				RawPointIndices;
-	FWordBulkData				LegacyRawPointIndices;
+	/** Accessor for the RawPointIndice which is editor only data: array of the original point (wedge) indices for each of the vertices in a FSkeletalMeshLODModel */
+	ENGINE_API const TArray<uint32>& GetRawPointIndices() const
+	{
+		return RawPointIndices2;
+	}
+
+	/** Accessor for the RawPointIndice which is editor only data: array of the original point (wedge) indices for each of the vertices in a FSkeletalMeshLODModel */
+	ENGINE_API TArray<uint32>& GetRawPointIndices()
+	{
+		return RawPointIndices2;
+	}
+	
+	UE_DEPRECATED(5.0, "Please do not access this function anymore. This data is not use anymore.")
+	ENGINE_API FRawSkeletalMeshBulkData& GetRawSkeletalMeshBulkData_DEPRECATED()
+	{
+		return RawSkeletalMeshBulkData_DEPRECATED;
+	}
 
-	/** Imported raw mesh data. Optional, only the imported mesh LOD has this, generated LOD or old asset will be null. */
-	FRawSkeletalMeshBulkData	RawSkeletalMeshBulkData_DEPRECATED;
 	/** This ID is use to create the DDC key, it must be set when we save the FRawSkeletalMeshBulkData. */
 	FString						RawSkeletalMeshBulkDataID;
 	bool						bIsBuildDataAvailable;
@@ -435,6 +447,13 @@ public:
 	int32 FindMeshInfoIndex(FName Name) const;
 
 private:
+	/** Editor only data: array of the original point (wedge) indices for each of the vertices in a FSkeletalMeshLODModel */
+	TArray<uint32>				RawPointIndices2;
+
+	FIntBulkData				RawPointIndices_DEPRECATED;
+	FWordBulkData				LegacyRawPointIndices_DEPRECATED;
+	FRawSkeletalMeshBulkData	RawSkeletalMeshBulkData_DEPRECATED;
+
 	//Mutex use by the CopyStructure function. It's a pointer because FCriticalSection privatize the operator= function, which will prevent this class operator= to use the default.
 	//We want to avoid having a custom equal operator that will get deprecated if dev forget to add the new member in this class
 	//The CopyStructure function will copy everything but make sure the destination mutex is set to a new mutex pointer.
