Fix the alternate skinning build to adjust the section max bone influence count in case the alternate skinning use more influence then the base skinning.
#rnx

diff --git a/Engine/Source/Developer/SkeletalMeshUtilitiesCommon/Private/LODUtilities.cpp b/Engine/Source/Developer/SkeletalMeshUtilitiesCommon/Private/LODUtilities.cpp
--- a/Engine/Source/Developer/SkeletalMeshUtilitiesCommon/Private/LODUtilities.cpp
+++ b/Engine/Source/Developer/SkeletalMeshUtilitiesCommon/Private/LODUtilities.cpp
@@ -1893,32 +1893,35 @@ bool FLODUtilities::UpdateAlternateSkinWeights(FSkeletalMeshLODModel& LODModelDe
 			TArray<int32> SrcWedgeIndexes;
 			FindWedgeIndexesUsingVertexIndex(ImportDataSrc, VertexIndexSrc, SrcWedgeIndexes);
 
-			//Check if we have a point that is perfectly matching (position, UV, material and vertex color). Because normals and tangent are on the triangles we do not test those.
-			for (int32 MatchDestinationIndex = 0; MatchDestinationIndex < SimilarDestinationVertex.Num(); ++MatchDestinationIndex)
+			if (SrcWedgeIndexes.Num() > 0 && SimilarDestinationVertex.Num() > 1)
 			{
-				int32 VertexIndexDest = SimilarDestinationVertex[MatchDestinationIndex];
-				TArray<int32> DestWedgeIndexes;
-				FindWedgeIndexesUsingVertexIndex(ImportDataDest, VertexIndexDest, DestWedgeIndexes);
-				for (int32 IndexDest = 0; IndexDest < DestWedgeIndexes.Num(); ++IndexDest)
+				//Check if we have a point that is perfectly matching (position, UV, material and vertex color). Because normals and tangent are on the triangles we do not test those.
+				for (int32 MatchDestinationIndex = 0; MatchDestinationIndex < SimilarDestinationVertex.Num(); ++MatchDestinationIndex)
 				{
-					int32 DestWedgeIndex = DestWedgeIndexes[IndexDest];
-					const SkeletalMeshImportData::FVertex& WedgeDest = ImportDataDest.Wedges[DestWedgeIndex];
-					for (int32 IndexSrc = 0; IndexSrc < SrcWedgeIndexes.Num(); ++IndexSrc)
+					int32 VertexIndexDest = SimilarDestinationVertex[MatchDestinationIndex];
+					TArray<int32> DestWedgeIndexes;
+					FindWedgeIndexesUsingVertexIndex(ImportDataDest, VertexIndexDest, DestWedgeIndexes);
+					for (int32 IndexDest = 0; IndexDest < DestWedgeIndexes.Num(); ++IndexDest)
 					{
-						int32 SrcWedgeIndex = SrcWedgeIndexes[IndexSrc];
-						const SkeletalMeshImportData::FVertex& WedgeSrc = ImportDataSrc.Wedges[SrcWedgeIndex];
-						//Wedge == operator test: material, vertex color and UVs
-						if (WedgeDest == WedgeSrc)
+						int32 DestWedgeIndex = DestWedgeIndexes[IndexDest];
+						const SkeletalMeshImportData::FVertex& WedgeDest = ImportDataDest.Wedges[DestWedgeIndex];
+						for (int32 IndexSrc = 0; IndexSrc < SrcWedgeIndexes.Num(); ++IndexSrc)
+						{
+							int32 SrcWedgeIndex = SrcWedgeIndexes[IndexSrc];
+							const SkeletalMeshImportData::FVertex& WedgeSrc = ImportDataSrc.Wedges[SrcWedgeIndex];
+							//Wedge == operator test: material, vertex color and UVs
+							if (WedgeDest == WedgeSrc)
+							{
+								VertexMatchDest.VertexIndexes.Add(SimilarDestinationVertex[MatchDestinationIndex]);
+								VertexMatchDest.Ratios.Add(1.0f);
+								break;
+							}
+						}
+						if (VertexMatchDest.VertexIndexes.Num() > 0)
 						{
-							VertexMatchDest.VertexIndexes.Add(SimilarDestinationVertex[MatchDestinationIndex]);
-							VertexMatchDest.Ratios.Add(1.0f);
 							break;
 						}
 					}
-					if (VertexMatchDest.VertexIndexes.Num() > 0)
-					{
-						break;
-					}
 				}
 			}
 			//If there is no direct match, simply put everything
@@ -2212,7 +2215,7 @@ bool FLODUtilities::UpdateAlternateSkinWeights(USkeletalMesh* SkeletalMeshDest,
 	return true;
 }
 
-void FLODUtilities::GenerateImportedSkinWeightProfileData(const FSkeletalMeshLODModel& LODModelDest, FImportedSkinWeightProfileData &ImportedProfileData)
+void FLODUtilities::GenerateImportedSkinWeightProfileData(FSkeletalMeshLODModel& LODModelDest, FImportedSkinWeightProfileData &ImportedProfileData)
 {
 	//Add the override buffer with the alternate influence data
 	TArray<FSoftSkinVertex> DestinationSoftVertices;
@@ -2223,6 +2226,7 @@ void FLODUtilities::GenerateImportedSkinWeightProfileData(const FSkeletalMeshLOD
 
 	//Get the maximum allow bone influence, so we can cut lowest weight properly and get the same result has the sk build
 	const int32 MaxInfluenceCount = FGPUBaseSkinVertexFactory::UseUnlimitedBoneInfluences(MAX_TOTAL_INFLUENCES) ? MAX_TOTAL_INFLUENCES : EXTRA_BONE_INFLUENCES;
+	int32 MaxNumInfluences = 0;
 
 	for (int32 VertexInstanceIndex = 0; VertexInstanceIndex < DestinationSoftVertices.Num(); ++VertexInstanceIndex)
 	{
@@ -2233,7 +2237,8 @@ void FLODUtilities::GenerateImportedSkinWeightProfileData(const FSkeletalMeshLOD
 		{
 			continue;
 		}
-		const TArray<FBoneIndexType> SectionBoneMap = LODModelDest.Sections[SectionIndex].BoneMap;
+		FSkelMeshSection& Section = LODModelDest.Sections[SectionIndex];
+		const TArray<FBoneIndexType> SectionBoneMap = Section.BoneMap;
 		const FSoftSkinVertex& Vertex = DestinationSoftVertices[VertexInstanceIndex];
 		const int32 VertexIndex = LODModelDest.MeshToImportVertexMap[VertexInstanceIndex];
 		check(VertexIndex >= 0 && VertexIndex <= LODModelDest.MaxImportVertex);
@@ -2277,6 +2282,15 @@ void FLODUtilities::GenerateImportedSkinWeightProfileData(const FSkeletalMeshLOD
 				break;
 			}
 		}
+		//Adjust section influence count if the alternate influence bone count is greater
+		if (InfluenceBoneIndex > MaxNumInfluences)
+		{
+			MaxNumInfluences = InfluenceBoneIndex;
+			if (MaxNumInfluences > Section.GetMaxBoneInfluences())
+			{
+				Section.MaxBoneInfluences = MaxNumInfluences;
+			}
+		}
 		//Use the same code has the build where we modify the index 0 to have a sum of 255 for all influence per skin vertex
 		SkinWeight.InfluenceWeights[0] += 255 - TotalInfluenceWeight;
 	}
diff --git a/Engine/Source/Developer/SkeletalMeshUtilitiesCommon/Public/LODUtilities.h b/Engine/Source/Developer/SkeletalMeshUtilitiesCommon/Public/LODUtilities.h
--- a/Engine/Source/Developer/SkeletalMeshUtilitiesCommon/Public/LODUtilities.h
+++ b/Engine/Source/Developer/SkeletalMeshUtilitiesCommon/Public/LODUtilities.h
@@ -176,7 +176,7 @@ private:
 	FLODUtilities() {}
 
 	/** Generate the editor-only data stored for a skin weight profile (relies on bone indices) */
-	static void GenerateImportedSkinWeightProfileData(const FSkeletalMeshLODModel& LODModelDest, FImportedSkinWeightProfileData& ImportedProfileData);
+	static void GenerateImportedSkinWeightProfileData(FSkeletalMeshLODModel& LODModelDest, FImportedSkinWeightProfileData& ImportedProfileData);
 	
 
 	/**
