- Alternate skinning now use interchange (if interchange is enabled)
- Custom LOD import now use interchange (if interchange is enabled)
- Add an InterchangeManager API function to retrieve per asset type extensions format
- Add an interchange mesh utilities to create a common place to add mesh import code.
- Add an interchange file picker interface which can be implemented by plugin
- Add an editor plugin that implement the file picker, The file picker can be customize to import specific asset type.

#rnx

diff --git a/Engine/Plugins/Experimental/Interchange/Editor/Source/Utilities/Private/InterchangeOpenFileDialog.cpp b/Engine/Plugins/Experimental/Interchange/Editor/Source/Utilities/Private/InterchangeOpenFileDialog.cpp
--- a/Engine/Plugins/Experimental/Interchange/Editor/Source/Utilities/Private/InterchangeOpenFileDialog.cpp
+++ b/Engine/Plugins/Experimental/Interchange/Editor/Source/Utilities/Private/InterchangeOpenFileDialog.cpp
@@ -64,48 +64,39 @@ namespace UE::Interchange::Utilities::Private
 		return ExtensionStrAll;
 	}
 
-	bool FilePickerDialog(const FString& Extensions, FString& OutFilename)
+	bool FilePickerDialog(const FString& Extensions, FInterchangeFilePickerParameters& Parameters, TArray<FString>& OutFilenames)
 	{
 		// First, display the file open dialog for selecting the file.
-		TArray<FString> OpenFilenames;
 		IDesktopPlatform* DesktopPlatform = FDesktopPlatformModule::Get();
-		bool bOpen = false;
 		if (DesktopPlatform)
 		{
-			FText PromptTitle = NSLOCTEXT("InterchangeUtilities_OpenFileDialog", "FilePickerDialog", "Select a file");
-			bOpen = DesktopPlatform->OpenFileDialog(
+			FText PromptTitle = Parameters.Title.IsEmpty() ? NSLOCTEXT("InterchangeUtilities_OpenFileDialog", "FilePickerDialog", "Select a file") : Parameters.Title;
+
+			const EFileDialogFlags::Type DialogFlags = Parameters.bAllowMultipleFiles ? EFileDialogFlags::Multiple : EFileDialogFlags::None;
+
+			return DesktopPlatform->OpenFileDialog(
 				FSlateApplication::Get().FindBestParentWindowHandleForDialogs(nullptr),
 				PromptTitle.ToString(),
-				TEXT(""),
+				Parameters.DefaultPath,
 				TEXT(""),
 				*Extensions,
-				EFileDialogFlags::None,
-				OpenFilenames
+				DialogFlags,
+				OutFilenames
 			);
 		}
-		FString SelectFilename;
-		// Only continue if we pressed OK and have only one file selected.
-		if (bOpen)
-		{
-			if (OpenFilenames.Num() >= 1)
-			{
-				OutFilename = OpenFilenames[0];
-				return true;
-			}
-		}
 
 		return false;
 	}
 } //ns UE::Interchange::Utilities::Private
 
-bool UInterchangeFilePickerGeneric::FilePickerForTranslatorAssetType(const EInterchangeTranslatorAssetType TranslatorAssetType, FString& OutFilename)
+bool UInterchangeFilePickerGeneric::FilePickerForTranslatorAssetType(const EInterchangeTranslatorAssetType TranslatorAssetType, FInterchangeFilePickerParameters& Parameters, TArray<FString>& OutFilenames)
 {
 	FString Extensions = UE::Interchange::Utilities::Private::GetOpenFileDialogExtensions(UInterchangeManager::GetInterchangeManager().GetSupportedAssetTypeFormats(TranslatorAssetType));
-	return UE::Interchange::Utilities::Private::FilePickerDialog(Extensions, OutFilename);
+	return UE::Interchange::Utilities::Private::FilePickerDialog(Extensions, Parameters, OutFilenames);
 }
 
-bool UInterchangeFilePickerGeneric::FilePickerForTranslatorType(const EInterchangeTranslatorType TranslatorType, FString& OutFilename)
+bool UInterchangeFilePickerGeneric::FilePickerForTranslatorType(const EInterchangeTranslatorType TranslatorType, FInterchangeFilePickerParameters& Parameters, TArray<FString>& OutFilenames)
 {
 	FString Extensions = UE::Interchange::Utilities::Private::GetOpenFileDialogExtensions(UInterchangeManager::GetInterchangeManager().GetSupportedFormats(TranslatorType));
-	return UE::Interchange::Utilities::Private::FilePickerDialog(Extensions, OutFilename);
+	return UE::Interchange::Utilities::Private::FilePickerDialog(Extensions, Parameters, OutFilenames);
 }
diff --git a/Engine/Plugins/Experimental/Interchange/Editor/Source/Utilities/Public/InterchangeOpenFileDialog.h b/Engine/Plugins/Experimental/Interchange/Editor/Source/Utilities/Public/InterchangeOpenFileDialog.h
--- a/Engine/Plugins/Experimental/Interchange/Editor/Source/Utilities/Public/InterchangeOpenFileDialog.h
+++ b/Engine/Plugins/Experimental/Interchange/Editor/Source/Utilities/Public/InterchangeOpenFileDialog.h
@@ -19,6 +19,6 @@ public:
 
 protected:
 
-	virtual bool FilePickerForTranslatorAssetType(const EInterchangeTranslatorAssetType TranslatorAssetType, FString& OutFilename) override;
-	virtual bool FilePickerForTranslatorType(const EInterchangeTranslatorType TranslatorType, FString& OutFilename) override;
+	virtual bool FilePickerForTranslatorAssetType(const EInterchangeTranslatorAssetType TranslatorAssetType, FInterchangeFilePickerParameters& Parameters, TArray<FString>& OutFilenames) override;
+	virtual bool FilePickerForTranslatorType(const EInterchangeTranslatorType TranslatorType, FInterchangeFilePickerParameters& Parameters, TArray<FString>& OutFilenames) override;
 };
diff --git a/Engine/Plugins/Experimental/Interchange/Runtime/Source/Pipelines/Private/InterchangeGenericAnimationPipeline.cpp b/Engine/Plugins/Experimental/Interchange/Runtime/Source/Pipelines/Private/InterchangeGenericAnimationPipeline.cpp
--- a/Engine/Plugins/Experimental/Interchange/Runtime/Source/Pipelines/Private/InterchangeGenericAnimationPipeline.cpp
+++ b/Engine/Plugins/Experimental/Interchange/Runtime/Source/Pipelines/Private/InterchangeGenericAnimationPipeline.cpp
@@ -41,14 +41,25 @@ namespace UE::Interchange::Private
 	}
 }
 
-void UInterchangeGenericAnimationPipeline::SetupReimportData(TObjectPtr<UObject> ReimportObject)
+void UInterchangeGenericAnimationPipeline::AdjustSettingsForReimportType(EInterchangeReimportType ImportType, TObjectPtr<UObject> ReimportAsset)
 {
 	check(!CommonSkeletalMeshesAndAnimationsProperties.IsNull());
-	if (UAnimSequence* AnimSequence = Cast<UAnimSequence>(ReimportObject))
+	if (ImportType == EInterchangeReimportType::AssetCustomLODImport
+		|| ImportType == EInterchangeReimportType::AssetCustomLODImport
+		|| ImportType == EInterchangeReimportType::AssetAlternateSkinningImport
+		|| ImportType == EInterchangeReimportType::AssetAlternateSkinningReimport)
 	{
-		//Set the skeleton to the current asset skeleton and re-import only the animation
-		CommonSkeletalMeshesAndAnimationsProperties->Skeleton = AnimSequence->GetSkeleton();
-		CommonSkeletalMeshesAndAnimationsProperties->bImportOnlyAnimations = true;
+		bImportAnimations = false;
+		CommonSkeletalMeshesAndAnimationsProperties->bImportOnlyAnimations = false;
+	}
+	else if(ImportType == EInterchangeReimportType::AssetReimport)
+	{
+		if (UAnimSequence* AnimSequence = Cast<UAnimSequence>(ReimportAsset))
+		{
+			//Set the skeleton to the current asset skeleton and re-import only the animation
+			CommonSkeletalMeshesAndAnimationsProperties->Skeleton = AnimSequence->GetSkeleton();
+			CommonSkeletalMeshesAndAnimationsProperties->bImportOnlyAnimations = true;
+		}
 	}
 }
 
diff --git a/Engine/Plugins/Experimental/Interchange/Runtime/Source/Pipelines/Private/InterchangeGenericAssetsPipeline.cpp b/Engine/Plugins/Experimental/Interchange/Runtime/Source/Pipelines/Private/InterchangeGenericAssetsPipeline.cpp
--- a/Engine/Plugins/Experimental/Interchange/Runtime/Source/Pipelines/Private/InterchangeGenericAssetsPipeline.cpp
+++ b/Engine/Plugins/Experimental/Interchange/Runtime/Source/Pipelines/Private/InterchangeGenericAssetsPipeline.cpp
@@ -106,26 +106,26 @@ bool UInterchangeGenericAssetsPipeline::IsSettingsAreValid() const
 	return Super::IsSettingsAreValid();
 }
 
-void UInterchangeGenericAssetsPipeline::SetupReimportData(TObjectPtr<UObject> ReimportObject)
+void UInterchangeGenericAssetsPipeline::AdjustSettingsForReimportType(EInterchangeReimportType ImportType, TObjectPtr<UObject> ReimportAsset)
 {
 	if (TexturePipeline)
 	{
-		TexturePipeline->SetupReimportData(ReimportObject);
+		TexturePipeline->AdjustSettingsForReimportType(ImportType, ReimportAsset);
 	}
 
 	if (MaterialPipeline)
 	{
-		MaterialPipeline->SetupReimportData(ReimportObject);
+		MaterialPipeline->AdjustSettingsForReimportType(ImportType, ReimportAsset);
 	}
 
 	if (MeshPipeline)
 	{
-		MeshPipeline->SetupReimportData(ReimportObject);
+		MeshPipeline->AdjustSettingsForReimportType(ImportType, ReimportAsset);
 	}
 
 	if (AnimationPipeline)
 	{
-		AnimationPipeline->SetupReimportData(ReimportObject);
+		AnimationPipeline->AdjustSettingsForReimportType(ImportType, ReimportAsset);
 	}
 }
 
diff --git a/Engine/Plugins/Experimental/Interchange/Runtime/Source/Pipelines/Private/InterchangeGenericMaterialPipeline.cpp b/Engine/Plugins/Experimental/Interchange/Runtime/Source/Pipelines/Private/InterchangeGenericMaterialPipeline.cpp
--- a/Engine/Plugins/Experimental/Interchange/Runtime/Source/Pipelines/Private/InterchangeGenericMaterialPipeline.cpp
+++ b/Engine/Plugins/Experimental/Interchange/Runtime/Source/Pipelines/Private/InterchangeGenericMaterialPipeline.cpp
@@ -40,6 +40,19 @@
 #include "UObject/Object.h"
 #include "UObject/ObjectMacros.h"
 
+void UInterchangeGenericMaterialPipeline::AdjustSettingsForReimportType(EInterchangeReimportType ImportType, TObjectPtr<UObject> ReimportAsset)
+{
+	bool bIsObjectAMaterial = ReimportAsset.IsNull() ? false : ReimportAsset->IsA(UMaterialInterface::StaticClass());
+	if ((!bIsObjectAMaterial && ImportType == EInterchangeReimportType::AssetReimport)
+		|| ImportType == EInterchangeReimportType::AssetCustomLODImport
+		|| ImportType == EInterchangeReimportType::AssetCustomLODReimport
+		|| ImportType == EInterchangeReimportType::AssetAlternateSkinningImport
+		|| ImportType == EInterchangeReimportType::AssetAlternateSkinningReimport)
+	{
+		MaterialImport = EInterchangeMaterialImportOption::DoNotImport;
+	}
+}
+
 void UInterchangeGenericMaterialPipeline::ExecutePreImportPipeline(UInterchangeBaseNodeContainer* InBaseNodeContainer, const TArray<UInterchangeSourceData*>& InSourceDatas)
 {
 	if (!InBaseNodeContainer)
diff --git a/Engine/Plugins/Experimental/Interchange/Runtime/Source/Pipelines/Private/InterchangeGenericMeshPipeline.cpp b/Engine/Plugins/Experimental/Interchange/Runtime/Source/Pipelines/Private/InterchangeGenericMeshPipeline.cpp
--- a/Engine/Plugins/Experimental/Interchange/Runtime/Source/Pipelines/Private/InterchangeGenericMeshPipeline.cpp
+++ b/Engine/Plugins/Experimental/Interchange/Runtime/Source/Pipelines/Private/InterchangeGenericMeshPipeline.cpp
@@ -13,6 +13,24 @@
 #include "Nodes/InterchangeBaseNode.h"
 #include "Nodes/InterchangeBaseNodeContainer.h"
 
+void UInterchangeGenericMeshPipeline::AdjustSettingsForReimportType(EInterchangeReimportType ImportType, TObjectPtr<UObject> ReimportAsset)
+{
+	check(!CommonSkeletalMeshesAndAnimationsProperties.IsNull());
+	if(ImportType == EInterchangeReimportType::AssetCustomLODImport
+		|| ImportType == EInterchangeReimportType::AssetCustomLODImport
+		|| ImportType == EInterchangeReimportType::AssetAlternateSkinningImport
+		|| ImportType == EInterchangeReimportType::AssetAlternateSkinningReimport)
+	{
+		bCreatePhysicsAsset = false;
+		PhysicsAsset = nullptr;
+	}
+	else if (USkeletalMesh* SkeletalMesh = Cast<USkeletalMesh>(ReimportAsset))
+	{
+		//Set the skeleton to the current asset skeleton
+		CommonSkeletalMeshesAndAnimationsProperties->Skeleton = SkeletalMesh->GetSkeleton();
+	}
+}
+
 void UInterchangeGenericMeshPipeline::ExecutePreImportPipeline(UInterchangeBaseNodeContainer* InBaseNodeContainer, const TArray<UInterchangeSourceData*>& InSourceDatas)
 {
 	if (!InBaseNodeContainer)
diff --git a/Engine/Plugins/Experimental/Interchange/Runtime/Source/Pipelines/Private/InterchangeGenericTexturePipeline.cpp b/Engine/Plugins/Experimental/Interchange/Runtime/Source/Pipelines/Private/InterchangeGenericTexturePipeline.cpp
--- a/Engine/Plugins/Experimental/Interchange/Runtime/Source/Pipelines/Private/InterchangeGenericTexturePipeline.cpp
+++ b/Engine/Plugins/Experimental/Interchange/Runtime/Source/Pipelines/Private/InterchangeGenericTexturePipeline.cpp
@@ -77,6 +77,19 @@ namespace UE::Interchange::Private
 #endif
 }
 
+void UInterchangeGenericTexturePipeline::AdjustSettingsForReimportType(EInterchangeReimportType ImportType, TObjectPtr<UObject> ReimportAsset)
+{
+	bool bIsObjectATexture = ReimportAsset.IsNull() ? false : ReimportAsset.IsA(UTexture::StaticClass());
+	if( (!bIsObjectATexture && ImportType == EInterchangeReimportType::AssetReimport)
+		|| ImportType == EInterchangeReimportType::AssetCustomLODImport
+		|| ImportType == EInterchangeReimportType::AssetCustomLODReimport
+		|| ImportType == EInterchangeReimportType::AssetAlternateSkinningImport
+		|| ImportType == EInterchangeReimportType::AssetAlternateSkinningReimport)
+	{
+		bImportTextures = false;
+	}
+}
+
 void UInterchangeGenericTexturePipeline::ExecutePreImportPipeline(UInterchangeBaseNodeContainer* InBaseNodeContainer, const TArray<UInterchangeSourceData*>& InSourceDatas)
 {
 	if (!InBaseNodeContainer)
diff --git a/Engine/Plugins/Experimental/Interchange/Runtime/Source/Pipelines/Public/InterchangeGenericAnimationPipeline.h b/Engine/Plugins/Experimental/Interchange/Runtime/Source/Pipelines/Public/InterchangeGenericAnimationPipeline.h
--- a/Engine/Plugins/Experimental/Interchange/Runtime/Source/Pipelines/Public/InterchangeGenericAnimationPipeline.h
+++ b/Engine/Plugins/Experimental/Interchange/Runtime/Source/Pipelines/Public/InterchangeGenericAnimationPipeline.h
@@ -100,7 +100,7 @@ public:
 	UPROPERTY()
 	FString SourceAnimationName;
 
-	virtual void SetupReimportData(TObjectPtr<UObject> ReimportObject) override;
+	virtual void AdjustSettingsForReimportType(EInterchangeReimportType ImportType, TObjectPtr<UObject> ReimportAsset) override;
 	
 protected:
 
diff --git a/Engine/Plugins/Experimental/Interchange/Runtime/Source/Pipelines/Public/InterchangeGenericAssetsPipeline.h b/Engine/Plugins/Experimental/Interchange/Runtime/Source/Pipelines/Public/InterchangeGenericAssetsPipeline.h
--- a/Engine/Plugins/Experimental/Interchange/Runtime/Source/Pipelines/Public/InterchangeGenericAssetsPipeline.h
+++ b/Engine/Plugins/Experimental/Interchange/Runtime/Source/Pipelines/Public/InterchangeGenericAssetsPipeline.h
@@ -80,8 +80,8 @@ public:
 
 	virtual bool IsSettingsAreValid() const override;
 
-	virtual void SetupReimportData(TObjectPtr<UObject> ReimportObject) override;
 
+	virtual void AdjustSettingsForReimportType(EInterchangeReimportType ImportType, TObjectPtr<UObject> ReimportAsset) override;
 protected:
 
 	virtual void ExecutePreImportPipeline(UInterchangeBaseNodeContainer* InBaseNodeContainer, const TArray<UInterchangeSourceData*>& InSourceDatas) override;
diff --git a/Engine/Plugins/Experimental/Interchange/Runtime/Source/Pipelines/Public/InterchangeGenericMaterialPipeline.h b/Engine/Plugins/Experimental/Interchange/Runtime/Source/Pipelines/Public/InterchangeGenericMaterialPipeline.h
--- a/Engine/Plugins/Experimental/Interchange/Runtime/Source/Pipelines/Public/InterchangeGenericMaterialPipeline.h
+++ b/Engine/Plugins/Experimental/Interchange/Runtime/Source/Pipelines/Public/InterchangeGenericMaterialPipeline.h
@@ -40,6 +40,8 @@ public:
 	UPROPERTY(EditAnywhere, BlueprintReadWrite, Category = "Materials", Meta= (EditCondition="MaterialImport==EInterchangeMaterialImportOption::ImportAsMaterialInstances", AllowedClasses="MaterialInterface"))
 	FSoftObjectPath ParentMaterial;
 
+	virtual void AdjustSettingsForReimportType(EInterchangeReimportType ImportType, TObjectPtr<UObject> ReimportAsset) override;
+
 protected:
 	virtual void ExecutePreImportPipeline(UInterchangeBaseNodeContainer* InBaseNodeContainer, const TArray<UInterchangeSourceData*>& InSourceDatas) override;
 
diff --git a/Engine/Plugins/Experimental/Interchange/Runtime/Source/Pipelines/Public/InterchangeGenericMeshPipeline.h b/Engine/Plugins/Experimental/Interchange/Runtime/Source/Pipelines/Public/InterchangeGenericMeshPipeline.h
--- a/Engine/Plugins/Experimental/Interchange/Runtime/Source/Pipelines/Public/InterchangeGenericMeshPipeline.h
+++ b/Engine/Plugins/Experimental/Interchange/Runtime/Source/Pipelines/Public/InterchangeGenericMeshPipeline.h
@@ -101,6 +101,8 @@ public:
 	UPROPERTY(EditAnywhere, BlueprintReadWrite, Category = "Skeletal Meshes", meta = (editcondition = "!bCreatePhysicsAsset"))
 	TObjectPtr<UPhysicsAsset> PhysicsAsset;
 
+	virtual void AdjustSettingsForReimportType(EInterchangeReimportType ImportType, TObjectPtr<UObject> ReimportAsset) override;
+
 protected:
 	virtual void ExecutePreImportPipeline(UInterchangeBaseNodeContainer* InBaseNodeContainer, const TArray<UInterchangeSourceData*>& InSourceDatas) override;
 
diff --git a/Engine/Plugins/Experimental/Interchange/Runtime/Source/Pipelines/Public/InterchangeGenericTexturePipeline.h b/Engine/Plugins/Experimental/Interchange/Runtime/Source/Pipelines/Public/InterchangeGenericTexturePipeline.h
--- a/Engine/Plugins/Experimental/Interchange/Runtime/Source/Pipelines/Public/InterchangeGenericTexturePipeline.h
+++ b/Engine/Plugins/Experimental/Interchange/Runtime/Source/Pipelines/Public/InterchangeGenericTexturePipeline.h
@@ -56,6 +56,7 @@ public:
 #endif
 
 public:
+	virtual void AdjustSettingsForReimportType(EInterchangeReimportType ImportType, TObjectPtr<UObject> ReimportAsset) override;
 	virtual void ExecutePreImportPipeline(UInterchangeBaseNodeContainer* InBaseNodeContainer, const TArray<UInterchangeSourceData*>& InSourceDatas) override;
 	virtual void ExecutePostImportPipeline(const UInterchangeBaseNodeContainer* InBaseNodeContainer, const FString& NodeKey, UObject* CreatedAsset, bool bIsAReimport) override;
 
diff --git a/Engine/Source/Developer/SkeletalMeshUtilitiesCommon/Private/LODUtilities.cpp b/Engine/Source/Developer/SkeletalMeshUtilitiesCommon/Private/LODUtilities.cpp
--- a/Engine/Source/Developer/SkeletalMeshUtilitiesCommon/Private/LODUtilities.cpp
+++ b/Engine/Source/Developer/SkeletalMeshUtilitiesCommon/Private/LODUtilities.cpp
@@ -1,6 +1,9 @@
 // Copyright Epic Games, Inc. All Rights Reserved.
 
 #include "LODUtilities.h"
+
+#if WITH_EDITOR
+
 #include "Misc/MessageDialog.h"
 #include "Misc/FeedbackContext.h"
 #include "Modules/ModuleManager.h"
@@ -13,6 +16,7 @@
 #include "Rendering/SkeletalMeshLODModel.h"
 #include "GenericQuadTree.h"
 #include "Engine/SkeletalMesh.h"
+#include "Engine/SkeletalMeshSocket.h"
 #include "EditorFramework/AssetImportData.h"
 #include "MeshUtilities.h"
 #include "MeshUtilitiesCommon.h"
@@ -38,6 +42,9 @@
 #include "Interfaces/ITargetPlatformManagerModule.h"
 #include "Misc/CoreMisc.h"
 
+#include "Framework/Notifications/NotificationManager.h"
+#include "Widgets/Notifications/SNotificationList.h"
+
 IMPLEMENT_MODULE(FDefaultModuleImpl, SkeletalMeshUtilitiesCommon)
 
 #define LOCTEXT_NAMESPACE "LODUtilities"
@@ -510,6 +517,392 @@ void FLODUtilities::RemoveLODs(FSkeletalMeshUpdateContext& UpdateContext, const
 	}
 }
 
+bool FLODUtilities::SetCustomLOD(USkeletalMesh* DestinationSkeletalMesh, USkeletalMesh* SourceSkeletalMesh, const int32 LodIndex, const FString& SourceDataFilename)
+{
+	if(!DestinationSkeletalMesh || !SourceSkeletalMesh)
+	{
+		return false;
+	}
+
+	FScopedSkeletalMeshPostEditChange ScopePostEditChange(DestinationSkeletalMesh);
+
+	//If the imported LOD already exist, we will need to reimport all the skin weight profiles
+	bool bMustReimportAlternateSkinWeightProfile = false;
+
+	// Get a list of all the clothing assets affecting this LOD so we can re-apply later
+	TArray<ClothingAssetUtils::FClothingAssetMeshBinding> ClothingBindings;
+	TArray<UClothingAssetBase*> ClothingAssetsInUse;
+	TArray<int32> ClothingAssetSectionIndices;
+	TArray<int32> ClothingAssetInternalLodIndices;
+
+	FSkeletalMeshModel* SourceImportedResource = SourceSkeletalMesh->GetImportedModel();
+	FSkeletalMeshModel* DestImportedResource = DestinationSkeletalMesh->GetImportedModel();
+
+	if (SourceImportedResource && SourceImportedResource->LODModels.IsValidIndex(LodIndex))
+	{
+		bMustReimportAlternateSkinWeightProfile = true;
+		FLODUtilities::UnbindClothingAndBackup(DestinationSkeletalMesh, ClothingBindings, LodIndex);
+	}
+
+	//Lambda to call to re-apply the clothing
+	auto ReapplyClothing = [&DestinationSkeletalMesh, &ClothingBindings, &SourceImportedResource, &LodIndex]()
+	{
+		if (SourceImportedResource && SourceImportedResource->LODModels.IsValidIndex(LodIndex))
+		{
+			// Re-apply our clothing assets
+			FLODUtilities::RestoreClothingFromBackup(DestinationSkeletalMesh, ClothingBindings, LodIndex);
+		}
+	};
+
+	// Now we copy the base FSkeletalMeshLODModel from the imported skeletal mesh as the new LOD in the selected mesh.
+	if (SourceImportedResource->LODModels.Num() == 0)
+	{
+		return false;
+	}
+
+	// Names of root bones must match.
+	// If the names of root bones don't match, the LOD Mesh does not share skeleton with base Mesh. 
+	if (SourceSkeletalMesh->GetRefSkeleton().GetBoneName(0) != DestinationSkeletalMesh->GetRefSkeleton().GetBoneName(0))
+	{
+		UE_LOG(LogLODUtilities, Error, TEXT("SkeletalMesh [%s] FLODUtilities::SetCustomLOD: Root bone in LOD is '%s' instead of '%s'.\nImport failed..")
+		, *DestinationSkeletalMesh->GetName()
+		, *SourceSkeletalMesh->GetRefSkeleton().GetBoneName(0).ToString()
+		, *DestinationSkeletalMesh->GetRefSkeleton().GetBoneName(0).ToString());
+		return false;
+	}
+	// We do some checking here that for every bone in the mesh we just imported, it's in our base ref skeleton, and the parent is the same.
+	for (int32 i = 0; i < SourceSkeletalMesh->GetRefSkeleton().GetRawBoneNum(); i++)
+	{
+		int32 LODBoneIndex = i;
+		FName LODBoneName = SourceSkeletalMesh->GetRefSkeleton().GetBoneName(LODBoneIndex);
+		int32 BaseBoneIndex = DestinationSkeletalMesh->GetRefSkeleton().FindBoneIndex(LODBoneName);
+		if (BaseBoneIndex == INDEX_NONE)
+		{
+			UE_LOG(LogLODUtilities, Error, TEXT("FLODUtilities::SetCustomLOD: Bone '%s' not found in destination SkeletalMesh '%s'.\nImport failed.")
+				, *LODBoneName.ToString()
+				, *DestinationSkeletalMesh->GetName());
+			return false;
+		}
+
+		if (i > 0)
+		{
+			int32 LODParentIndex = SourceSkeletalMesh->GetRefSkeleton().GetParentIndex(LODBoneIndex);
+			FName LODParentName = SourceSkeletalMesh->GetRefSkeleton().GetBoneName(LODParentIndex);
+
+			int32 BaseParentIndex = DestinationSkeletalMesh->GetRefSkeleton().GetParentIndex(BaseBoneIndex);
+			FName BaseParentName = DestinationSkeletalMesh->GetRefSkeleton().GetBoneName(BaseParentIndex);
+
+			if (LODParentName != BaseParentName)
+			{
+				UE_LOG(LogLODUtilities, Error, TEXT("SkeletalMesh [%s] FLODUtilities::SetCustomLOD: Bone '%s' in LOD has parent '%s' instead of '%s'")
+					, *DestinationSkeletalMesh->GetName()
+					, *LODBoneName.ToString()
+					, *LODParentName.ToString()
+					, *BaseParentName.ToString());
+				return false;
+			}
+		}
+	}
+
+	FScopedSkeletalMeshPostEditChange ScopedPostEditChange(DestinationSkeletalMesh);
+
+	FSkeletalMeshLODModel& NewLODModel = SourceImportedResource->LODModels[0];
+
+	// If this LOD is not the base LOD, we check all bones we have sockets on are present in it.
+	if (LodIndex > 0)
+	{
+		const TArray<USkeletalMeshSocket*>& Sockets = DestinationSkeletalMesh->GetMeshOnlySocketList();
+
+		for (int32 i = 0; i < Sockets.Num(); i++)
+		{
+			// Find bone index the socket is attached to.
+			USkeletalMeshSocket* Socket = Sockets[i];
+			int32 SocketBoneIndex = SourceSkeletalMesh->GetRefSkeleton().FindBoneIndex(Socket->BoneName);
+
+			// If this LOD does not contain the socket bone, abort import.
+			if (SocketBoneIndex == INDEX_NONE)
+			{
+				UE_LOG(LogLODUtilities, Error, TEXT("FLODUtilities::SetCustomLOD: This LOD is missing bone '%s' used by socket '%s'.\nAborting import.")
+					, *Socket->BoneName.ToString()
+					, *Socket->SocketName.ToString());
+				return false;
+			}
+		}
+	}
+
+	{
+		//The imported LOD is always in LOD 0 of the SourceSkeletalMesh
+		const int32 SourceLODIndex = 0;
+		if (!SourceSkeletalMesh->IsLODImportedDataEmpty(SourceLODIndex))
+		{
+			// Fix up the imported data bone indexes
+			FSkeletalMeshImportData LODImportData;
+			SourceSkeletalMesh->LoadLODImportedData(SourceLODIndex, LODImportData);
+			const int32 LODImportDataBoneNumber = LODImportData.RefBonesBinary.Num();
+			//We want to create a remap array so we can fix all influence easily
+			TArray<int32> ImportDataBoneRemap;
+			ImportDataBoneRemap.AddZeroed(LODImportDataBoneNumber);
+			//We generate a new RefBonesBinary array to replace the existing one
+			TArray<SkeletalMeshImportData::FBone> RemapedRefBonesBinary;
+			RemapedRefBonesBinary.AddZeroed(DestinationSkeletalMesh->GetRefSkeleton().GetNum());
+			for (int32 ImportBoneIndex = 0; ImportBoneIndex < LODImportDataBoneNumber; ++ImportBoneIndex)
+			{
+				SkeletalMeshImportData::FBone& ImportedBone = LODImportData.RefBonesBinary[ImportBoneIndex];
+				int32 LODBoneIndex = ImportBoneIndex;
+				FName LODBoneName = FName(*ImportedBone.Name);
+				int32 BaseBoneIndex = DestinationSkeletalMesh->GetRefSkeleton().FindBoneIndex(LODBoneName);
+				ImportDataBoneRemap[ImportBoneIndex] = BaseBoneIndex;
+				if (BaseBoneIndex != INDEX_NONE)
+				{
+					RemapedRefBonesBinary[BaseBoneIndex] = ImportedBone;
+					if (RemapedRefBonesBinary[BaseBoneIndex].ParentIndex != INDEX_NONE)
+					{
+						RemapedRefBonesBinary[BaseBoneIndex].ParentIndex = ImportDataBoneRemap[RemapedRefBonesBinary[BaseBoneIndex].ParentIndex];
+					}
+				}
+			}
+			//Copy the new RefBonesBinary over the existing one
+			LODImportData.RefBonesBinary = RemapedRefBonesBinary;
+
+			//Fix the influences
+			bool bNeedShrinking = false;
+			const int32 InfluenceNumber = LODImportData.Influences.Num();
+			for (int32 InfluenceIndex = InfluenceNumber - 1; InfluenceIndex >= 0; --InfluenceIndex)
+			{
+				SkeletalMeshImportData::FRawBoneInfluence& Influence = LODImportData.Influences[InfluenceIndex];
+				Influence.BoneIndex = ImportDataBoneRemap[Influence.BoneIndex];
+				if (Influence.BoneIndex == INDEX_NONE)
+				{
+					const int32 DeleteCount = 1;
+					const bool AllowShrink = false;
+					LODImportData.Influences.RemoveAt(InfluenceIndex, DeleteCount, AllowShrink);
+					bNeedShrinking = true;
+				}
+			}
+			//Shrink the array if we have deleted at least one entry
+			if (bNeedShrinking)
+			{
+				LODImportData.Influences.Shrink();
+			}
+			//Save the fix up remap bone index
+			SourceSkeletalMesh->SaveLODImportedData(SourceLODIndex, LODImportData);
+		}
+
+		// Fix up the ActiveBoneIndices array.
+		for (int32 ActiveIndex = 0; ActiveIndex < NewLODModel.ActiveBoneIndices.Num(); ActiveIndex++)
+		{
+			int32 LODBoneIndex = NewLODModel.ActiveBoneIndices[ActiveIndex];
+			FName LODBoneName = SourceSkeletalMesh->GetRefSkeleton().GetBoneName(LODBoneIndex);
+			int32 BaseBoneIndex = DestinationSkeletalMesh->GetRefSkeleton().FindBoneIndex(LODBoneName);
+			NewLODModel.ActiveBoneIndices[ActiveIndex] = BaseBoneIndex;
+		}
+
+		// Fix up the chunk BoneMaps.
+		for (int32 SectionIndex = 0; SectionIndex < NewLODModel.Sections.Num(); SectionIndex++)
+		{
+			FSkelMeshSection& Section = NewLODModel.Sections[SectionIndex];
+			for (int32 BoneMapIndex = 0; BoneMapIndex < Section.BoneMap.Num(); BoneMapIndex++)
+			{
+				int32 LODBoneIndex = Section.BoneMap[BoneMapIndex];
+				FName LODBoneName = SourceSkeletalMesh->GetRefSkeleton().GetBoneName(LODBoneIndex);
+				int32 BaseBoneIndex = DestinationSkeletalMesh->GetRefSkeleton().FindBoneIndex(LODBoneName);
+				Section.BoneMap[BoneMapIndex] = BaseBoneIndex;
+			}
+		}
+
+		// Create the RequiredBones array in the LODModel from the ref skeleton.
+		for (int32 RequiredBoneIndex = 0; RequiredBoneIndex < NewLODModel.RequiredBones.Num(); RequiredBoneIndex++)
+		{
+			FName LODBoneName = SourceSkeletalMesh->GetRefSkeleton().GetBoneName(NewLODModel.RequiredBones[RequiredBoneIndex]);
+			int32 BaseBoneIndex = DestinationSkeletalMesh->GetRefSkeleton().FindBoneIndex(LODBoneName);
+			if (BaseBoneIndex != INDEX_NONE)
+			{
+				NewLODModel.RequiredBones[RequiredBoneIndex] = BaseBoneIndex;
+			}
+			else
+			{
+				NewLODModel.RequiredBones.RemoveAt(RequiredBoneIndex--);
+			}
+		}
+
+		// Also sort the RequiredBones array to be strictly increasing.
+		NewLODModel.RequiredBones.Sort();
+		DestinationSkeletalMesh->GetRefSkeleton().EnsureParentsExistAndSort(NewLODModel.ActiveBoneIndices);
+	}
+
+	// To be extra-nice, we apply the difference between the root transform of the meshes to the verts.
+	FMatrix44f LODToBaseTransform = FMatrix44f(SourceSkeletalMesh->GetRefPoseMatrix(0).InverseFast() * DestinationSkeletalMesh->GetRefPoseMatrix(0));
+
+	for (int32 SectionIndex = 0; SectionIndex < NewLODModel.Sections.Num(); SectionIndex++)
+	{
+		FSkelMeshSection& Section = NewLODModel.Sections[SectionIndex];
+
+		// Fix up soft verts.
+		for (int32 i = 0; i < Section.SoftVertices.Num(); i++)
+		{
+			Section.SoftVertices[i].Position = LODToBaseTransform.TransformPosition(Section.SoftVertices[i].Position);
+			Section.SoftVertices[i].TangentX = LODToBaseTransform.TransformVector(Section.SoftVertices[i].TangentX);
+			Section.SoftVertices[i].TangentY = LODToBaseTransform.TransformVector(Section.SoftVertices[i].TangentY);
+			Section.SoftVertices[i].TangentZ = LODToBaseTransform.TransformVector(Section.SoftVertices[i].TangentZ);
+		}
+	}
+
+	//Restore the LOD section data in case this LOD was reimport and some material match
+	if (DestImportedResource->LODModels.IsValidIndex(LodIndex) && SourceImportedResource->LODModels.IsValidIndex(0))
+	{
+		const TArray<FSkelMeshSection>& ExistingSections = DestImportedResource->LODModels[LodIndex].Sections;
+		const FSkeletalMeshLODInfo& ExistingInfo = *(DestinationSkeletalMesh->GetLODInfo(LodIndex));
+
+		TArray<FSkelMeshSection>& ImportedSections = SourceImportedResource->LODModels[0].Sections;
+		const FSkeletalMeshLODInfo& ImportedInfo = *(SourceSkeletalMesh->GetLODInfo(0));
+
+		auto GetImportMaterialSlotName = [](const USkeletalMesh* SkelMesh, const FSkelMeshSection& Section, int32 SectionIndex, const FSkeletalMeshLODInfo& Info, int32& OutMaterialIndex)->FName
+		{
+			const TArray<FSkeletalMaterial>& MeshMaterials = SkelMesh->GetMaterials();
+			check(MeshMaterials.Num() > 0);
+			OutMaterialIndex = Section.MaterialIndex;
+			if (Info.LODMaterialMap.IsValidIndex(SectionIndex) && Info.LODMaterialMap[SectionIndex] != INDEX_NONE)
+			{
+				OutMaterialIndex = Info.LODMaterialMap[SectionIndex];
+			}
+			FName ImportedMaterialSlotName = NAME_None;
+			if (MeshMaterials.IsValidIndex(OutMaterialIndex))
+			{
+				ImportedMaterialSlotName = MeshMaterials[OutMaterialIndex].ImportedMaterialSlotName;
+			}
+			else
+			{
+				ImportedMaterialSlotName = MeshMaterials[0].ImportedMaterialSlotName;
+				OutMaterialIndex = 0;
+			}
+			return ImportedMaterialSlotName;
+		};
+
+		for (int32 ExistingSectionIndex = 0; ExistingSectionIndex < ExistingSections.Num(); ++ExistingSectionIndex)
+		{
+			const FSkelMeshSection& ExistingSection = ExistingSections[ExistingSectionIndex];
+			int32 ExistingMaterialIndex = 0;
+			FName ExistingImportedMaterialSlotName = GetImportMaterialSlotName(DestinationSkeletalMesh, ExistingSection, ExistingSectionIndex, ExistingInfo, ExistingMaterialIndex);
+
+			for (int32 ImportedSectionIndex = 0; ImportedSectionIndex < ImportedSections.Num(); ++ImportedSectionIndex)
+			{
+				FSkelMeshSection& ImportedSection = ImportedSections[ImportedSectionIndex];
+				int32 ImportedMaterialIndex = 0;
+				FName ImportedImportedMaterialSlotName = GetImportMaterialSlotName(SourceSkeletalMesh, ImportedSection, ImportedSectionIndex, ImportedInfo, ImportedMaterialIndex);
+				if (ExistingImportedMaterialSlotName != NAME_None)
+				{
+					if (ImportedImportedMaterialSlotName == ExistingImportedMaterialSlotName)
+					{
+						//Set the value and exit
+						ImportedSection.bCastShadow = ExistingSection.bCastShadow;
+						ImportedSection.bVisibleInRayTracing = ExistingSection.bVisibleInRayTracing;
+						ImportedSection.bRecomputeTangent = ExistingSection.bRecomputeTangent;
+						ImportedSection.RecomputeTangentsVertexMaskChannel = ExistingSection.RecomputeTangentsVertexMaskChannel;
+						break;
+					}
+				}
+				else if (SourceSkeletalMesh->GetMaterials()[ImportedMaterialIndex] == DestinationSkeletalMesh->GetMaterials()[ExistingMaterialIndex]) //Use material slot compare to match in case the name is none
+				{
+					//Set the value and exit
+					ImportedSection.bCastShadow = ExistingSection.bCastShadow;
+					ImportedSection.bVisibleInRayTracing = ExistingSection.bVisibleInRayTracing;
+					ImportedSection.bRecomputeTangent = ExistingSection.bRecomputeTangent;
+					ImportedSection.RecomputeTangentsVertexMaskChannel = ExistingSection.RecomputeTangentsVertexMaskChannel;
+					break;
+				}
+			}
+		}
+	}
+
+	// If we want to add this as a new LOD to this mesh - add to LODModels/LODInfo array.
+	if (LodIndex == DestImportedResource->LODModels.Num())
+	{
+		DestImportedResource->LODModels.Add(new FSkeletalMeshLODModel());
+
+		// Add element to LODInfo array.
+		DestinationSkeletalMesh->AddLODInfo();
+		check(DestinationSkeletalMesh->GetLODNum() == DestImportedResource->LODModels.Num());
+	}
+
+	// Set up LODMaterialMap to number of materials in new mesh.
+	FSkeletalMeshLODInfo& LODInfo = *(DestinationSkeletalMesh->GetLODInfo(LodIndex));
+
+	//Copy the build settings
+	if (SourceSkeletalMesh->GetLODInfo(0))
+	{
+		const FSkeletalMeshLODInfo& ImportedLODInfo = *(SourceSkeletalMesh->GetLODInfo(0));
+		LODInfo.BuildSettings = ImportedLODInfo.BuildSettings;
+	}
+
+	TArray<FSkeletalMaterial>& BaseMaterials = DestinationSkeletalMesh->GetMaterials();
+	LODInfo.LODMaterialMap.Empty();
+	// Now set up the material mapping array.
+	for (int32 SectionIndex = 0; SectionIndex < NewLODModel.Sections.Num(); SectionIndex++)
+	{
+		int32 MatIdx = NewLODModel.Sections[SectionIndex].MaterialIndex;
+		// Try and find the auto-assigned material in the array.
+		int32 LODMatIndex = INDEX_NONE;
+		//First try to match by name
+		for (int32 BaseMaterialIndex = 0; BaseMaterialIndex < BaseMaterials.Num(); ++BaseMaterialIndex)
+		{
+			const FSkeletalMaterial& SkeletalMaterial = BaseMaterials[BaseMaterialIndex];
+			if (SkeletalMaterial.ImportedMaterialSlotName != NAME_None && SkeletalMaterial.ImportedMaterialSlotName == SourceSkeletalMesh->GetMaterials()[MatIdx].ImportedMaterialSlotName)
+			{
+				LODMatIndex = BaseMaterialIndex;
+				break;
+			}
+		}
+
+		// If we dont have a match, add a new entry to the material list.
+		if (LODMatIndex == INDEX_NONE)
+		{
+			LODMatIndex = BaseMaterials.Add(SourceSkeletalMesh->GetMaterials()[MatIdx]);
+		}
+
+		LODInfo.LODMaterialMap.Add(LODMatIndex);
+	}
+
+	// Release all resources before replacing the model
+	DestinationSkeletalMesh->PreEditChange(NULL);
+
+	// Assign new FSkeletalMeshLODModel to desired slot in selected skeletal mesh.
+	FSkeletalMeshLODModel::CopyStructure(&(DestImportedResource->LODModels[LodIndex]), &NewLODModel);
+	//Copy the import data into the base skeletalmesh for the imported LOD
+	USkeletalMesh::CopyImportedData(0, SourceSkeletalMesh, LodIndex, DestinationSkeletalMesh);
+
+
+	// If this LOD had been generated previously by automatic mesh reduction, clear that flag.
+	LODInfo.bHasBeenSimplified = false;
+	if (DestinationSkeletalMesh->GetLODSettings() == nullptr || !DestinationSkeletalMesh->GetLODSettings()->HasValidSettings() || DestinationSkeletalMesh->GetLODSettings()->GetNumberOfSettings() <= LodIndex)
+	{
+		//Make sure any custom LOD have correct settings (no reduce)
+		LODInfo.ReductionSettings.NumOfTrianglesPercentage = 1.0f;
+		LODInfo.ReductionSettings.NumOfVertPercentage = 1.0f;
+		LODInfo.ReductionSettings.MaxDeviationPercentage = 0.0f;
+	}
+
+	// Set LOD source filename
+	DestinationSkeletalMesh->GetLODInfo(LodIndex)->SourceImportFilename = UAssetImportData::SanitizeImportFilename(SourceDataFilename, nullptr);
+	DestinationSkeletalMesh->GetLODInfo(LodIndex)->bImportWithBaseMesh = false;
+
+	ReapplyClothing();
+
+	//Must be the last step because it cleanup the fbx importer to import the alternate skinning FBX
+	if (bMustReimportAlternateSkinWeightProfile)
+	{
+		//TODO port skin weights utilities outside of UnrealEd module
+		//FSkinWeightsUtilities::ReimportAlternateSkinWeight(DestinationSkeletalMesh, LodIndex);
+	}
+	
+	// Notification of success
+	FNotificationInfo NotificationInfo(FText::GetEmpty());
+	NotificationInfo.Text = FText::Format(NSLOCTEXT("UnrealEd", "LODImportSuccessful", "Mesh for LOD {0} imported successfully!"), FText::AsNumber(LodIndex));
+	NotificationInfo.ExpireDuration = 5.0f;
+	FSlateNotificationManager::Get().AddNotification(NotificationInfo);
+	
+	return true;
+}
+
 /** Given three direction vectors, indicates if A and B are on the same 'side' of Vec. */
 bool VectorsOnSameSide(const FVector2D& Vec, const FVector2D& A, const FVector2D& B)
 {
@@ -3438,4 +3831,6 @@ bool FLODUtilities::StripLODGeometry(USkeletalMesh* SkeletalMesh, const int32 LO
 	return true;
 }
 
-#undef LOCTEXT_NAMESPACE // "LODUtilities"
\ No newline at end of file
+#undef LOCTEXT_NAMESPACE // "LODUtilities"
+
+#endif //WITH_EDITOR
\ No newline at end of file
diff --git a/Engine/Source/Developer/SkeletalMeshUtilitiesCommon/Public/LODUtilities.h b/Engine/Source/Developer/SkeletalMeshUtilitiesCommon/Public/LODUtilities.h
--- a/Engine/Source/Developer/SkeletalMeshUtilitiesCommon/Public/LODUtilities.h
+++ b/Engine/Source/Developer/SkeletalMeshUtilitiesCommon/Public/LODUtilities.h
@@ -2,6 +2,8 @@
 
 #pragma once
 
+#if WITH_EDITOR
+
 #include "CoreMinimal.h"
 #include "MeshBuild.h"
 #include "Engine/SkeletalMesh.h"
@@ -10,6 +12,7 @@
 #include "Framework/Commands/UIAction.h"
 #include "Animation/MorphTarget.h"
 
+
 namespace ClothingAssetUtils
 {
 	struct FClothingAssetMeshBinding;
@@ -80,6 +83,16 @@ public:
 	*/
 	static void RemoveLODs(FSkeletalMeshUpdateContext& UpdateContext, const TArray<int32>& DesiredLODs);
 
+	/*
+	 * Add or change the LOD data specified by LodIndex with the content of the sourceSkeletalMesh.
+	 *
+	 * @Param DestinationSkeletalMesh - Skeletal mesh receiving the LOD data.
+	 * @Param SourceSkeletalMesh - Skeletal mesh providing the LOD data we want to use to add or modify the specified destination LOD at LodIndex. The LOD we want to import is always the base LOD of the source mesh.
+	 * @Param LodIndex - The destination lod index we want to add or modify.
+	 * @Param SourceDataFilename - The file name we use to import the source skeletal mesh.
+	 */
+	static bool SetCustomLOD(USkeletalMesh* DestinationSkeletalMesh, USkeletalMesh* SourceSkeletalMesh, const int32 LodIndex, const FString& SourceDataFilename);
+
 	/**
 	*	Simplifies the static mesh based upon various user settings for DesiredLOD.
 	*
@@ -151,7 +164,7 @@ public:
 	
 	static void RestoreClothingFromBackup(USkeletalMesh* SkeletalMesh, TArray<ClothingAssetUtils::FClothingAssetMeshBinding>& ClothingBindings);
 	static void RestoreClothingFromBackup(USkeletalMesh* SkeletalMesh, TArray<ClothingAssetUtils::FClothingAssetMeshBinding>& ClothingBindings, const int32 LODIndex);
-	
+
 	/**
 	 * Before building skeletalmesh base LOD (LOD index 0) using MeshUtilities.BuildSkeletalMesh, we want to adjust the imported faces material index to point on the correct sk material. We use the material name to match the material.
 	 * @param Materials - The skeletalmesh material list to fit the import data face material
@@ -177,7 +190,6 @@ private:
 
 	/** Generate the editor-only data stored for a skin weight profile (relies on bone indices) */
 	static void GenerateImportedSkinWeightProfileData(FSkeletalMeshLODModel& LODModelDest, FImportedSkinWeightProfileData& ImportedProfileData);
-	
 
 	/**
 	 *	Simplifies the static mesh based upon various user settings for DesiredLOD
@@ -205,3 +217,5 @@ private:
 	*/
 	static void ClearGeneratedMorphTarget(USkeletalMesh* SkeletalMesh, int32 DesiredLOD);
 };
+
+#endif //WITH_EDITOR
\ No newline at end of file
diff --git a/Engine/Source/Editor/Persona/Private/SkinWeightProfileHelpers.cpp b/Engine/Source/Editor/Persona/Private/SkinWeightProfileHelpers.cpp
--- a/Engine/Source/Editor/Persona/Private/SkinWeightProfileHelpers.cpp
+++ b/Engine/Source/Editor/Persona/Private/SkinWeightProfileHelpers.cpp
@@ -31,7 +31,7 @@ void FSkinWeightProfileHelpers::ImportSkinWeightProfile(USkeletalMesh* InSkeleta
 	if (InSkeletalMesh)
 	{
 		const int32 LOD0Index = 0;
-		FString PickedFileName = FSkinWeightsUtilities::PickSkinWeightFBXPath(LOD0Index, InSkeletalMesh);
+		FString PickedFileName = FSkinWeightsUtilities::PickSkinWeightPath(LOD0Index, InSkeletalMesh);
 
 		bool bShouldImport = false;
 		USkinWeightImportOptions* ImportSettings = GetMutableDefault<USkinWeightImportOptions>();		
@@ -103,7 +103,7 @@ void FSkinWeightProfileHelpers::ImportSkinWeightProfileLOD(USkeletalMesh* InSkel
 	if (InSkeletalMesh)
 	{
 		// Pick a FBX file to import the weights from 
-		const FString PickedFileName = FSkinWeightsUtilities::PickSkinWeightFBXPath(LODIndex, InSkeletalMesh);
+		const FString PickedFileName = FSkinWeightsUtilities::PickSkinWeightPath(LODIndex, InSkeletalMesh);
 		FScopedSuspendAlternateSkinWeightPreview ScopedSuspendAlternateSkinnWeightPreview(InSkeletalMesh);
 		FScopedSkeletalMeshPostEditChange ScopedPostEditChange(InSkeletalMesh);
 		// Try and import skin weights for a specific mesh LOD
@@ -154,7 +154,7 @@ void FSkinWeightProfileHelpers::ReimportSkinWeightProfileLOD(USkeletalMesh* InSk
 				if (EAppReturnType::Yes == FMessageDialog::Open(EAppMsgType::YesNo, WarningMessage))
 				{
 					// Otherwise let the user pick a new path
-					const FString PickedFileName = FSkinWeightsUtilities::PickSkinWeightFBXPath(LODIndex, InSkeletalMesh);
+					const FString PickedFileName = FSkinWeightsUtilities::PickSkinWeightPath(LODIndex, InSkeletalMesh);
 					bResult = FSkinWeightsUtilities::ImportAlternateSkinWeight(InSkeletalMesh, PickedFileName, LODIndex, InProfileName);
 				}
 			}
diff --git a/Engine/Source/Editor/UnrealEd/Private/FbxMeshUtils.cpp b/Engine/Source/Editor/UnrealEd/Private/FbxMeshUtils.cpp
--- a/Engine/Source/Editor/UnrealEd/Private/FbxMeshUtils.cpp
+++ b/Engine/Source/Editor/UnrealEd/Private/FbxMeshUtils.cpp
@@ -25,6 +25,13 @@
 #include "ImportUtils/SkeletalMeshImportUtils.h"
 #include "ImportUtils/StaticMeshImportUtils.h"
 
+#include "Async/Async.h"
+#include "InterchangeAssetImportData.h"
+#include "InterchangeManager.h"
+#include "InterchangeMeshUtilities.h"
+#include "InterchangeProjectSettings.h"
+#include "Settings/EditorExperimentalSettings.h"
+
 #if WITH_APEX_CLOTHING
 	#include "ApexClothingUtils.h"
 #endif // #if WITH_APEX_CLOTHING
@@ -90,9 +97,50 @@ namespace FbxMeshUtils
 
 	bool ImportStaticMeshLOD( UStaticMesh* BaseStaticMesh, const FString& Filename, int32 LODLevel)
 	{
+		if (!BaseStaticMesh)
+		{
+			UE_LOG(LogExportMeshUtils, Log, TEXT("Cannot import custom LOD because the staticmesh is NULL."));
+			return false;
+		}
+
+		//We will use interchange only if interchange is enabled and the mesh we want to add a LOD was imported with interchange
+		const UEditorExperimentalSettings* EditorExperimentalSettings = GetDefault<UEditorExperimentalSettings>();
+		const UInterchangeAssetImportData* SelectedInterchangeAssetImportData = Cast<UInterchangeAssetImportData>(BaseStaticMesh->GetAssetImportData());
+		if (EditorExperimentalSettings->bEnableInterchangeFramework && SelectedInterchangeAssetImportData)
+		{
+			UInterchangeSourceData* SourceData = UInterchangeManager::GetInterchangeManager().CreateSourceData(Filename);
+			//Call interchange mesh utilities to import custom LOD
+			UInterchangeMeshUtilities::ImportCustomLodAsync(BaseStaticMesh, LODLevel, SourceData).Then([BaseStaticMesh, LODLevel](TFuture<bool> Result)
+				{
+					bool bResult = Result.Get();
+					Async(EAsyncExecution::TaskGraphMainThread, [BaseStaticMesh, LODLevel, bResult]()
+						{
+							if (bResult)
+							{
+								// Notification of success
+								FNotificationInfo NotificationInfo(FText::GetEmpty());
+								NotificationInfo.Text = FText::Format(NSLOCTEXT("UnrealEd", "LODImportSuccessful", "Static mesh LOD {0} imported successfully!"), FText::AsNumber(LODLevel));
+								NotificationInfo.ExpireDuration = 5.0f;
+								FSlateNotificationManager::Get().AddNotification(NotificationInfo);
+							}
+							else
+							{
+								// Notification of failure
+								FNotificationInfo NotificationInfo(FText::GetEmpty());
+								NotificationInfo.Text = FText::Format(NSLOCTEXT("UnrealEd", "LODImportFail", "Failed to import static mesh LOD {0}!"), FText::AsNumber(LODLevel));
+								NotificationInfo.ExpireDuration = 5.0f;
+								FSlateNotificationManager::Get().AddNotification(NotificationInfo);
+							}
+						});
+				});
+
+			return true;
+		}
+
 		bool bSuccess = false;
 
 		UE_LOG(LogExportMeshUtils, Log, TEXT("Fbx LOD loading"));
+
 		// logger for all error/warnings
 		// this one prints all messages that are stored in FFbxImporter
 		// this function seems to get called outside of FBX factory
@@ -243,14 +291,48 @@ namespace FbxMeshUtils
 
 	bool ImportSkeletalMeshLOD( class USkeletalMesh* SelectedSkelMesh, const FString& Filename, int32 LODLevel)
 	{
-		UnFbx::FFbxImporter* FFbxImporter = UnFbx::FFbxImporter::GetInstance();
 		//Make sure skeletal mesh is valid
 		if (!SelectedSkelMesh)
 		{
-			FFbxImporter->AddTokenizedErrorMessage(FTokenizedMessage::Create(EMessageSeverity::Error, LOCTEXT("FBXImport_NoSelectedSkeletalMesh", "Cannot import a LOD if there is not a valid selected skeletal mesh.")), FFbxErrors::Generic_MeshNotFound);
+			UE_LOG(LogExportMeshUtils, Error, TEXT("Cannot import a LOD if there is not a valid selected skeletal mesh."));
 			return false;
 		}
 
+		//We will use interchange only if interchange is enable and the skeletalmesh we want to add a LOD was import with interchange
+		const UEditorExperimentalSettings* EditorExperimentalSettings = GetDefault<UEditorExperimentalSettings>();
+		UInterchangeAssetImportData* SelectedInterchangeAssetImportData = Cast<UInterchangeAssetImportData>(SelectedSkelMesh->GetAssetImportData());
+		if (EditorExperimentalSettings->bEnableInterchangeFramework && SelectedInterchangeAssetImportData)
+		{
+			UInterchangeSourceData* SourceData = UInterchangeManager::GetInterchangeManager().CreateSourceData(Filename);
+			//Call interchange mesh utilities to import custom LOD
+			UInterchangeMeshUtilities::ImportCustomLodAsync(SelectedSkelMesh, LODLevel, SourceData).Then([SelectedSkelMesh, LODLevel](TFuture<bool> Result)
+				{
+					bool bResult = Result.Get();
+					Async(EAsyncExecution::TaskGraphMainThread, [SelectedSkelMesh, LODLevel, bResult]()
+						{
+							if (bResult)
+							{
+								// Notification of success
+								FNotificationInfo NotificationInfo(FText::GetEmpty());
+								NotificationInfo.Text = FText::Format(NSLOCTEXT("UnrealEd", "LODImportSuccessful", "Mesh for LOD {0} imported successfully!"), FText::AsNumber(LODLevel));
+								NotificationInfo.ExpireDuration = 5.0f;
+								FSlateNotificationManager::Get().AddNotification(NotificationInfo);
+							}
+							else
+							{
+								// Notification of failure
+								FNotificationInfo NotificationInfo(FText::GetEmpty());
+								NotificationInfo.Text = FText::Format(NSLOCTEXT("UnrealEd", "LODImportFail", "Failed to import mesh for LOD {0}!"), FText::AsNumber(LODLevel));
+								NotificationInfo.ExpireDuration = 5.0f;
+								FSlateNotificationManager::Get().AddNotification(NotificationInfo);
+							}
+						});
+				});
+			return true;
+		}
+
+		UnFbx::FFbxImporter* FFbxImporter = UnFbx::FFbxImporter::GetInstance();
+
 		bool bSuccess = false;
 
 		// Check the file extension for FBX. Anything that isn't .FBX is rejected
@@ -617,21 +699,17 @@ namespace FbxMeshUtils
 		USkeletalMesh* SkeletalMesh = Cast<USkeletalMesh>(SelectedMesh);
 		UStaticMesh* StaticMesh = Cast<UStaticMesh>(SelectedMesh);
 
-		if( !SkeletalMesh && !StaticMesh )
-		{
-			return false;
-		}
-
 		FString FilenameToImport("");
-
-		if(SkeletalMesh)
+		UInterchangeAssetImportData* SelectedInterchangeAssetImportData = nullptr;
+		if (SkeletalMesh)
 		{
-			if(SkeletalMesh->IsValidLODIndex(LODLevel))
+			if (SkeletalMesh->IsValidLODIndex(LODLevel))
 			{
 				FilenameToImport = SkeletalMesh->GetLODInfo(LODLevel)->SourceImportFilename.IsEmpty() ?
 					SkeletalMesh->GetLODInfo(LODLevel)->SourceImportFilename :
 					UAssetImportData::ResolveImportFilename(SkeletalMesh->GetLODInfo(LODLevel)->SourceImportFilename, nullptr);
 			}
+			SelectedInterchangeAssetImportData = Cast<UInterchangeAssetImportData>(SkeletalMesh->GetAssetImportData());
 		}
 		else if (StaticMesh)
 		{
@@ -642,12 +720,54 @@ namespace FbxMeshUtils
 					SourceModel.SourceImportFilename :
 					UAssetImportData::ResolveImportFilename(SourceModel.SourceImportFilename, nullptr);
 			}
+			SelectedInterchangeAssetImportData = Cast<UInterchangeAssetImportData>(StaticMesh->GetAssetImportData());
+		}
+		else
+		{
+			//We support only staticmesh and skeletalmesh asset for LOD import
+			return false;
 		}
 
 		// Check the file exists first
 		const bool bSourceFileExists = FPaths::FileExists(FilenameToImport);
 		// We'll give the user a chance to choose a new file if a previously set file fails to import
 		const bool bPromptOnFail = bSourceFileExists;
+
+		//We will use interchange only if interchange is enable and the skeletalmesh we want to add a LOD was import with interchange
+		const UEditorExperimentalSettings* EditorExperimentalSettings = GetDefault<UEditorExperimentalSettings>();
+		if(EditorExperimentalSettings->bEnableInterchangeFramework && SelectedInterchangeAssetImportData)
+		{
+			auto CustomLodImportContinuation = [bNotifyCB, SkeletalMesh, StaticMesh, LODLevel](TFuture<bool> Result)
+			{
+				bool bResult = Result.Get();
+				Async(EAsyncExecution::TaskGraphMainThread, [bNotifyCB, SkeletalMesh, StaticMesh, LODLevel, bResult]()
+					{
+						if (bNotifyCB && bResult)
+						{
+							if (SkeletalMesh)
+							{
+								GEditor->GetEditorSubsystem<UImportSubsystem>()->BroadcastAssetPostLODImport(SkeletalMesh, LODLevel);
+							}
+							else if (StaticMesh)
+							{
+								GEditor->GetEditorSubsystem<UImportSubsystem>()->BroadcastAssetPostLODImport(StaticMesh, LODLevel);
+							}
+						}
+					});
+			};
+
+			if(!bSourceFileExists)
+			{
+				//Call interchange mesh utilities to import custom LOD
+				UInterchangeMeshUtilities::ImportCustomLodAsync(SelectedMesh, LODLevel).Then(CustomLodImportContinuation);
+			}
+			else
+			{
+				UInterchangeSourceData* SourceData = UInterchangeManager::GetInterchangeManager().CreateSourceData(FilenameToImport);
+				UInterchangeMeshUtilities::ImportCustomLodAsync(SelectedMesh, LODLevel, SourceData).Then(CustomLodImportContinuation);
+			}
+			return true;
+		}
 		
 		if(!bSourceFileExists || FilenameToImport.IsEmpty())
 		{
diff --git a/Engine/Source/Editor/UnrealEd/Private/SkinWeightsUtilities.cpp b/Engine/Source/Editor/UnrealEd/Private/SkinWeightsUtilities.cpp
--- a/Engine/Source/Editor/UnrealEd/Private/SkinWeightsUtilities.cpp
+++ b/Engine/Source/Editor/UnrealEd/Private/SkinWeightsUtilities.cpp
@@ -32,7 +32,10 @@
 #include "DesktopPlatformModule.h"
 #include "EditorDirectories.h"
 #include "Framework/Application/SlateApplication.h"
+#include "InterchangeAssetImportData.h"
+#include "InterchangeFilePickerBase.h"
 #include "InterchangeManager.h"
+#include "InterchangeProjectSettings.h"
 #include "Interfaces/ITargetPlatform.h"
 #include "Interfaces/ITargetPlatformManagerModule.h"
 #include "Misc/CoreMisc.h"
@@ -109,7 +112,7 @@ bool FSkinWeightsUtilities::ImportAlternateSkinWeight(USkeletalMesh* SkeletalMes
 
 		//TODO create a pipeline that set all the proper skeletalmesh options (look at the legacy system setup)
 		UE::Interchange::FAssetImportResultRef AssetImportResult = InterchangeManager.ImportAssetAsync(ImportAssetPath, ScopedSourceData.GetSourceData(), ImportAssetParameters);
-		AssetImportResult->WaitUntilDone();
+		AssetImportResult->WaitUntilDone(); //TODO, do not stall the main thread here, WaitUntilDone will tick taskgraph so the job can complete even if we wait on the game thread.
 		if (USkeletalMesh* ImportedSkeletalMesh = Cast< USkeletalMesh >(AssetImportResult->GetFirstAssetOfClass(USkeletalMesh::StaticClass())))
 		{
 			ImportedObject = ImportedSkeletalMesh;
@@ -315,7 +318,7 @@ bool FSkinWeightsUtilities::ReimportAlternateSkinWeight(USkeletalMesh* SkeletalM
 		}
 		else
 		{
-			const FString PickedFileName = FSkinWeightsUtilities::PickSkinWeightFBXPath(TargetLODIndex, SkeletalMesh);
+			const FString PickedFileName = FSkinWeightsUtilities::PickSkinWeightPath(TargetLODIndex, SkeletalMesh);
 			if (!PickedFileName.IsEmpty() && FPaths::FileExists(PickedFileName))
 			{
 				bResult |= FSkinWeightsUtilities::ImportAlternateSkinWeight(SkeletalMesh, PickedFileName, TargetLODIndex, ProfileInfo.Name);
@@ -383,58 +386,112 @@ bool FSkinWeightsUtilities::RemoveSkinnedWeightProfileData(USkeletalMesh* Skelet
 	return bBuildSuccess;
 }
 
-FString FSkinWeightsUtilities::PickSkinWeightFBXPath(int32 LODIndex, USkeletalMesh* SkeletalMesh)
+FString FSkinWeightsUtilities::PickSkinWeightPath(int32 LODIndex, USkeletalMesh* SkeletalMesh)
 {
 	FString PickedFileName("");
+	
+	bool bUseInterchangeFramework = false;
+#if WITH_EDITOR
+	bUseInterchangeFramework = GetDefault<UEditorExperimentalSettings>()->bEnableInterchangeFramework;
+#endif
+	const UInterchangeAssetImportData* SelectedInterchangeAssetImportData = Cast<UInterchangeAssetImportData>(SkeletalMesh->GetAssetImportData());
 
-	FString ExtensionStr;
-	ExtensionStr += TEXT("FBX files|*.fbx|");
-
-	// First, display the file open dialog for selecting the file.
-	TArray<FString> OpenFilenames;
-	IDesktopPlatform* DesktopPlatform = FDesktopPlatformModule::Get();
-	bool bOpen = false;
-	if (DesktopPlatform)
+	if (bUseInterchangeFramework && SelectedInterchangeAssetImportData)
 	{
-		// Try and retrieve the path containing the original skeletal mesh source data, and set it as default path for the file dialog
-		UFbxSkeletalMeshImportData* ImportData = SkeletalMesh ? Cast<UFbxSkeletalMeshImportData>(SkeletalMesh->GetAssetImportData()) : nullptr;
-		FString DefaultPath;
-		FString TempString;
-		if (ImportData)
+		const FString& FirstSourceFile = SelectedInterchangeAssetImportData->ScriptGetFirstFilename();
+		FString DefaultPath = FPaths::GetPath(FirstSourceFile);
+		// Otherwise resort back to last imported directory
+		if (!FPaths::DirectoryExists(DefaultPath))
 		{
-			ImportData->GetImportContentFilename(DefaultPath, TempString);
-			DefaultPath = FPaths::GetPath(DefaultPath);
+			DefaultPath = FEditorDirectories::Get().GetLastDirectory(ELastDirectory::GENERIC_IMPORT);
 		}
-		
-		// Otherwise resort back to last FBX directory
-		if(!FPaths::DirectoryExists(DefaultPath))
-		{
-			DefaultPath = FEditorDirectories::Get().GetLastDirectory(ELastDirectory::FBX);
-		}		
-		
-		const FString DialogTitle = TEXT("Pick FBX file containing Skin Weight data for LOD ") + FString::FormatAsNumber(LODIndex);
-		bOpen = DesktopPlatform->OpenFileDialog(
-			FSlateApplication::Get().FindBestParentWindowHandleForDialogs(nullptr),
-			DialogTitle,
-			*DefaultPath,
-			TEXT(""),
-			*ExtensionStr,
-			EFileDialogFlags::None,
-			OpenFilenames
-		);
-	}
 
-	if (bOpen)
-	{
-		if (OpenFilenames.Num() == 1)
+		//Ask the user for a file path
+		const UInterchangeProjectSettings* InterchangeProjectSettings = GetDefault<UInterchangeProjectSettings>();
+		UInterchangeFilePickerBase* FilePicker = nullptr;
+
+		//In runtime we do not have any pipeline configurator
+#if WITH_EDITORONLY_DATA
+		TSoftClassPtr <UInterchangeFilePickerBase> FilePickerClass = InterchangeProjectSettings->FilePickerClass;
+		if (FilePickerClass.IsValid())
 		{
-			PickedFileName = OpenFilenames[0];
-			// Set last directory path for FBX files
-			FEditorDirectories::Get().SetLastDirectory(ELastDirectory::FBX, FPaths::GetPath(PickedFileName));
+			UClass* FilePickerClassLoaded = FilePickerClass.LoadSynchronous();
+			if (FilePickerClassLoaded)
+			{
+				FilePicker = NewObject<UInterchangeFilePickerBase>(GetTransientPackage(), FilePickerClassLoaded, NAME_None, RF_NoFlags);
+			}
 		}
-		else
+#endif
+		if (FilePicker)
+		{
+			FInterchangeFilePickerParameters Parameters;
+			Parameters.bAllowMultipleFiles = false;
+			Parameters.DefaultPath = DefaultPath;
+			Parameters.Title = FText::Format(NSLOCTEXT("FSkinWeightsUtilities", "PickSkinWeightPath_Title", "Choose a file to import alternate skinning for LOD {0}"), FText::AsNumber(LODIndex));
+			TArray<FString> Filenames;
+
+			if (FilePicker->ScriptedFilePickerForTranslatorAssetType(EInterchangeTranslatorAssetType::Meshes, Parameters, Filenames))
+			{
+				ensure(Filenames.Num() == 1);
+				PickedFileName = Filenames[0];
+				FEditorDirectories::Get().SetLastDirectory(ELastDirectory::GENERIC_IMPORT, FPaths::GetPath(PickedFileName));
+			}
+			else
+			{
+				// Error
+			}
+		}
+	}
+	else
+	{
+		bool bOpen = false;
+		TArray<FString> OpenFilenames;
+		FString ExtensionStr;
+		ExtensionStr += TEXT("FBX files|*.fbx|");
+
+		// First, display the file open dialog for selecting the file.
+		IDesktopPlatform* DesktopPlatform = FDesktopPlatformModule::Get();
+		if (DesktopPlatform)
 		{
-			// Error
+			// Try and retrieve the path containing the original skeletal mesh source data, and set it as default path for the file dialog
+			UFbxSkeletalMeshImportData* ImportData = SkeletalMesh ? Cast<UFbxSkeletalMeshImportData>(SkeletalMesh->GetAssetImportData()) : nullptr;
+			FString DefaultPath;
+			FString TempString;
+			if (ImportData)
+			{
+				ImportData->GetImportContentFilename(DefaultPath, TempString);
+				DefaultPath = FPaths::GetPath(DefaultPath);
+			}
+		
+			// Otherwise resort back to last FBX directory
+			if(!FPaths::DirectoryExists(DefaultPath))
+			{
+				DefaultPath = FEditorDirectories::Get().GetLastDirectory(ELastDirectory::FBX);
+			}		
+
+			const FString DialogTitle = TEXT("Pick FBX file containing Skin Weight data for LOD ") + FString::FormatAsNumber(LODIndex);
+			bOpen = DesktopPlatform->OpenFileDialog(
+				FSlateApplication::Get().FindBestParentWindowHandleForDialogs(nullptr),
+				DialogTitle,
+				*DefaultPath,
+				TEXT(""),
+				*ExtensionStr,
+				EFileDialogFlags::None,
+				OpenFilenames
+			);
+			if (bOpen)
+			{
+				if (OpenFilenames.Num() == 1)
+				{
+					PickedFileName = OpenFilenames[0];
+					// Set last directory path for FBX files
+					FEditorDirectories::Get().SetLastDirectory(ELastDirectory::FBX, FPaths::GetPath(PickedFileName));
+				}
+				else
+				{
+					// Error
+				}
+			}
 		}
 	}
 
diff --git a/Engine/Source/Editor/UnrealEd/Public/SkinWeightsUtilities.h b/Engine/Source/Editor/UnrealEd/Public/SkinWeightsUtilities.h
--- a/Engine/Source/Editor/UnrealEd/Public/SkinWeightsUtilities.h
+++ b/Engine/Source/Editor/UnrealEd/Public/SkinWeightsUtilities.h
@@ -42,9 +42,9 @@ public:
 	static bool RemoveSkinnedWeightProfileData(USkeletalMesh* SkeletalMesh, const FName& ProfileName, int32 LODIndex);
 
 	/*
-	 * Ask user a FBX file path for a particular LOD
+	 * Ask user a file path for a particular LOD
 	 */
-	static FString PickSkinWeightFBXPath(int32 LODIndex, USkeletalMesh* SkeletalMesh);
+	static FString PickSkinWeightPath(int32 LODIndex, USkeletalMesh* SkeletalMesh);
 private:
 	FSkinWeightsUtilities() {}
 };
diff --git a/Engine/Source/Runtime/Engine/Classes/Engine/StaticMesh.h b/Engine/Source/Runtime/Engine/Classes/Engine/StaticMesh.h
--- a/Engine/Source/Runtime/Engine/Classes/Engine/StaticMesh.h
+++ b/Engine/Source/Runtime/Engine/Classes/Engine/StaticMesh.h
@@ -1531,6 +1531,15 @@ public:
 	FOnExtendedBoundsChanged& GetOnExtendedBoundsChanged() { return OnExtendedBoundsChanged; }
 	FOnMeshChanged& GetOnMeshChanged() { return OnMeshChanged; }
 
+	/*
+	 * Add or change the LOD data specified by LodIndex with the content of the sourceStaticMesh.
+	 *
+	 * @Param SourceStaticMesh - The data we want to use to add or modify the specified lod
+	 * @Param LodIndex - The lod index we want to add or modify.
+	 * @Param SourceDataFilename - The source filename thta need to set into the LOD to allow re-import.
+	 */
+	ENGINE_API bool SetCustomLOD(const UStaticMesh* SourceStaticMesh, int32 LodIndex, const FString& SourceDataFilename);
+
 	//SourceModels API
 	ENGINE_API FStaticMeshSourceModel& AddSourceModel();
 
diff --git a/Engine/Source/Runtime/Engine/Private/StaticMesh.cpp b/Engine/Source/Runtime/Engine/Private/StaticMesh.cpp
--- a/Engine/Source/Runtime/Engine/Private/StaticMesh.cpp
+++ b/Engine/Source/Runtime/Engine/Private/StaticMesh.cpp
@@ -3929,6 +3929,126 @@ FStaticMeshSourceModel& UStaticMesh::GetSourceModel(int32 Index)
 	PRAGMA_ENABLE_DEPRECATION_WARNINGS
 }
 
+bool UStaticMesh::SetCustomLOD(const UStaticMesh* SourceStaticMesh, int32 DestinationLodIndex, const FString& SourceDataFilename)
+{
+
+#if WITH_EDITORONLY_DATA
+	if (!ensure(SourceStaticMesh) || SourceStaticMesh->GetNumSourceModels() <= 0)
+	{
+		return false;
+	}
+
+	const int32 SourceLodIndex = (SourceStaticMesh->GetNumSourceModels() > DestinationLodIndex) ? DestinationLodIndex : 0;
+	if(!SourceStaticMesh->IsSourceModelValid(SourceLodIndex))
+	{
+		return false;
+	}
+
+	const FMeshDescription* SourceMeshDescription = SourceStaticMesh->GetMeshDescription(SourceLodIndex);
+	if(!SourceMeshDescription)
+	{
+		return false;
+	}
+
+	if(DestinationLodIndex >= GetNumSourceModels())
+	{
+		// Add one LOD 
+		AddSourceModel();
+		if (GetNumSourceModels() <= DestinationLodIndex)
+		{
+			DestinationLodIndex = GetNumSourceModels() - 1;
+		}
+	}
+
+	//To restore the material and section data, we need to know which material the source is using
+	const TArray<FStaticMaterial>& SourceMaterials = SourceStaticMesh->GetStaticMaterials();
+	TMap<int32, FName> SourceImportedMaterialNameUsed;
+	{
+		FStaticMeshConstAttributes SourceAttributes(*SourceMeshDescription);
+		TPolygonGroupAttributesConstRef<FName> MaterialSlotNames = SourceAttributes.GetPolygonGroupMaterialSlotNames();
+		for (FPolygonGroupID PolygonGroupID : SourceMeshDescription->PolygonGroups().GetElementIDs())
+		{
+			const int32 SourceMaterialIndex = PolygonGroupID.GetValue();
+			check(SourceMaterials.IsValidIndex(SourceMaterialIndex));
+			SourceImportedMaterialNameUsed.FindOrAdd(SourceMaterialIndex) = MaterialSlotNames[PolygonGroupID];
+		}
+	}
+
+	TArray<FStaticMaterial>& DestinationMaterials = GetStaticMaterials();
+
+	FMeshDescription* DestinationMeshDescription = GetMeshDescription(DestinationLodIndex);
+	if (DestinationMeshDescription == nullptr)
+	{
+		DestinationMeshDescription = CreateMeshDescription(DestinationLodIndex);
+		check(DestinationMeshDescription != nullptr);
+		CommitMeshDescription(DestinationLodIndex);
+
+		//Make sure an imported mesh do not get reduce if there was no mesh data before reimport.
+		//In this case we have a generated LOD convert to a custom LOD
+		FStaticMeshSourceModel& SrcModel = GetSourceModel(DestinationLodIndex);
+		SrcModel.ReductionSettings.MaxDeviation = 0.0f;
+		SrcModel.ReductionSettings.PercentTriangles = 1.0f;
+		SrcModel.ReductionSettings.PercentVertices = 1.0f;
+	}
+	else
+	{
+		// clear out the old mesh data
+		DestinationMeshDescription->Empty();
+	}
+
+	//Make sure all materials use by the new LOD is pointing on a valid static material
+	for(const TPair<int32, FName>& SourceImportedMaterialNamePair : SourceImportedMaterialNameUsed)
+	{
+		bool bFoundMatch = false;
+		FName NameSearch = SourceImportedMaterialNamePair.Value;
+			
+		for (int32 MaterialIndex = 0; MaterialIndex < DestinationMaterials.Num(); ++MaterialIndex)
+		{
+			const FStaticMaterial& StaticMaterial = DestinationMaterials[MaterialIndex];
+			if (NameSearch == StaticMaterial.ImportedMaterialSlotName)
+			{
+				bFoundMatch = true;
+				break;
+			}
+		}
+
+		if (!bFoundMatch)
+		{
+			//Add the missing material slot
+			FStaticMaterial StaticMaterial(nullptr, NameSearch, NameSearch);
+		}
+	}
+
+	//Copy the mesh description of the source into the destination
+	*DestinationMeshDescription = *SourceMeshDescription;
+
+	if(IsInGameThread())
+	{
+		PostEditChange();
+		MarkPackageDirty();
+	}
+	else
+	{
+		UStaticMesh* ThisMesh = this;
+		Async(EAsyncExecution::TaskGraphMainThread, [ThisMesh]()
+		{
+			ThisMesh->PostEditChange();
+			ThisMesh->MarkPackageDirty();
+		});
+	}
+
+	if (IsSourceModelValid(DestinationLodIndex))
+	{
+		FStaticMeshSourceModel& SourceModel = GetSourceModel(DestinationLodIndex);
+		SourceModel.SourceImportFilename = UAssetImportData::SanitizeImportFilename(SourceDataFilename, nullptr);
+		SourceModel.bImportWithBaseMesh = false;
+	}
+	return true;
+#else
+	return false;
+#endif //!WITH_EDITORONLY_DATA
+}
+
 FStaticMeshSourceModel& UStaticMesh::AddSourceModel()
 {
 	WaitUntilAsyncPropertyReleased(EStaticMeshAsyncProperties::SourceModels);
diff --git a/Engine/Source/Runtime/Experimental/Interchange/Core/Public/InterchangePipelineBase.h b/Engine/Source/Runtime/Experimental/Interchange/Core/Public/InterchangePipelineBase.h
--- a/Engine/Source/Runtime/Experimental/Interchange/Core/Public/InterchangePipelineBase.h
+++ b/Engine/Source/Runtime/Experimental/Interchange/Core/Public/InterchangePipelineBase.h
@@ -20,6 +20,17 @@ enum class EInterchangePipelineTask : uint8
 	Export
 };
 
+UENUM()
+enum class EInterchangeReimportType : uint8
+{
+	AssetReimport,
+	SceneReimport,
+	AssetCustomLODImport, //The import for custom LOD is there because we use a copy of the asset import data pipeline stack.
+	AssetCustomLODReimport,
+	AssetAlternateSkinningImport, //The import for custom LOD is there because we use a copy of the asset import data pipeline stack.
+	AssetAlternateSkinningReimport,
+};
+
 UCLASS(BlueprintType, Blueprintable, Experimental, EditInlineNew, Abstract)
 class INTERCHANGECORE_API UInterchangePipelineBase : public UObject
 {
@@ -116,7 +127,7 @@ public:
 	/**
 	 * This function is called before showing the import dialog it is not called when doing a re-import.
 	 */
-	virtual void PreDialogCleanup(const FName PipelineStackName) {};
+	virtual void PreDialogCleanup(const FName PipelineStackName) {}
 
 	/**
 	 * This function should return true if all the pipeline settings are in a valid state to start the import.
@@ -128,9 +139,12 @@ public:
 	}
 
 	/**
-	 * This function is call before showing the re-import dialog it is not call when importing a new assets.
+	 * This function is call only when we do a re-import before we show the pipeline dialog. Pipeline that override it can change the existing settings according to the re-import type.
+	 * The function is also call when we import or re-import custom LOD and alternate skinning.
+	 * @Param ReimportType - Tell pipeline what re-import type the user want to achieve.
+	 * @Param ReimportAsset - This is an optional parameter which is set when re-importing an asset.
 	 */
-	virtual void SetupReimportData(TObjectPtr<UObject> ReimportObject) {};
+	virtual void AdjustSettingsForReimportType(EInterchangeReimportType ReimportType, TObjectPtr<UObject> ReimportAsset) {}
 
 	/**
 	 * This function is used to add the given message object directly into the results for this operation.
diff --git a/Engine/Source/Runtime/Experimental/Interchange/Engine/InterchangeEngine.Build.cs b/Engine/Source/Runtime/Experimental/Interchange/Engine/InterchangeEngine.Build.cs
--- a/Engine/Source/Runtime/Experimental/Interchange/Engine/InterchangeEngine.Build.cs
+++ b/Engine/Source/Runtime/Experimental/Interchange/Engine/InterchangeEngine.Build.cs
@@ -21,7 +21,8 @@ public class InterchangeEngine : ModuleRules
 				"AssetRegistry",
 				"SlateCore",
 				"Slate",
-				"DeveloperSettings"
+				"DeveloperSettings",
+				"SkeletalMeshUtilitiesCommon",
 			}
 		);
 	}
diff --git a/Engine/Source/Runtime/Experimental/Interchange/Engine/Private/InterchangeManager.cpp b/Engine/Source/Runtime/Experimental/Interchange/Engine/Private/InterchangeManager.cpp
--- a/Engine/Source/Runtime/Experimental/Interchange/Engine/Private/InterchangeManager.cpp
+++ b/Engine/Source/Runtime/Experimental/Interchange/Engine/Private/InterchangeManager.cpp
@@ -719,7 +719,7 @@ UInterchangeManager::ImportInternal(const FString& ContentPath, const UInterchan
 
 	if ( ImportAssetParameters.OverridePipelines.Num() == 0 )
 	{
-		const bool bIsUnattended = FApp::IsUnattended() || GIsAutomationTesting;
+		const bool bIsUnattended = FApp::IsUnattended() || GIsAutomationTesting || ImportAssetParameters.bIsAutomated;
 
 #if WITH_EDITORONLY_DATA
 		const bool bShowPipelineStacksConfigurationDialog = !bIsUnattended
@@ -756,7 +756,15 @@ UInterchangeManager::ImportInternal(const FString& ContentPath, const UInterchan
 				{
 					//Duplicate the pipeline saved in the asset import data
 					UInterchangePipelineBase* GeneratedPipeline = Cast<UInterchangePipelineBase>(StaticDuplicateObject(SourcePipeline, GetTransientPackage()));
-					GeneratedPipeline->SetupReimportData(ImportAssetParameters.ReimportAsset);
+
+					if (ImportType == UE::Interchange::EImportType::ImportType_Scene)
+					{
+						GeneratedPipeline->AdjustSettingsForReimportType(EInterchangeReimportType::SceneReimport, nullptr);
+					}
+					else
+					{
+						GeneratedPipeline->AdjustSettingsForReimportType(EInterchangeReimportType::AssetReimport, ImportAssetParameters.ReimportAsset);
+					}
 					PipelineStack.Add(GeneratedPipeline);
 				}
 				else
diff --git a/Engine/Source/Runtime/Experimental/Interchange/Engine/Private/InterchangeMeshUtilities.cpp b/Engine/Source/Runtime/Experimental/Interchange/Engine/Private/InterchangeMeshUtilities.cpp
new file mode 100644
--- /dev/null
+++ b/Engine/Source/Runtime/Experimental/Interchange/Engine/Private/InterchangeMeshUtilities.cpp
@@ -0,0 +1,166 @@
+// Copyright Epic Games, Inc. All Rights Reserved.
+
+#include "InterchangeMeshUtilities.h"
+
+#include "Async/Future.h"
+#include "Engine/SkeletalMesh.h"
+#include "Engine/StaticMesh.h"
+#include "InterchangeAssetImportData.h"
+#include "InterchangeFactoryBase.h"
+#include "InterchangeEngineLogPrivate.h"
+#include "InterchangeManager.h"
+#include "InterchangeProjectSettings.h"
+#include "InterchangeSourceData.h"
+#include "Logging/LogMacros.h"
+#include "UObject/Class.h"
+#include "UObject/Object.h"
+#include "UObject/ObjectMacros.h"
+#include "LODUtilities.h"
+
+TFuture<bool> UInterchangeMeshUtilities::ImportCustomLodAsync(UObject* MeshObject, const int32 LodIndex)
+{
+	TSharedPtr<TPromise<bool>> Promise = MakeShared<TPromise<bool>>();
+	if (!MeshObject)
+	{
+		UE_LOG(LogInterchangeEngine, Warning, TEXT("FInterchangeMeshUtilities::ImportCustomLod parameter MeshObject cannot be null."));
+		Promise->SetValue(false);
+		return Promise->GetFuture();
+	}
+	if (!IsInGameThread())
+	{
+		UE_LOG(LogInterchangeEngine, Warning, TEXT("FInterchangeMeshUtilities::ImportCustomLod Cannot ask user a file path outside of the game thread."));
+		Promise->SetValue(false);
+		return Promise->GetFuture();
+	}
+	UInterchangeManager& InterchangeManager = UInterchangeManager::GetInterchangeManager();
+
+	//Ask the user for a file path
+	const UInterchangeProjectSettings* InterchangeProjectSettings = GetDefault<UInterchangeProjectSettings>();
+	UInterchangeFilePickerBase* FilePicker = nullptr;
+
+	//In runtime we do not have any pipeline configurator
+#if WITH_EDITORONLY_DATA
+	TSoftClassPtr <UInterchangeFilePickerBase> FilePickerClass = InterchangeProjectSettings->FilePickerClass;
+	if (FilePickerClass.IsValid())
+	{
+		UClass* FilePickerClassLoaded = FilePickerClass.LoadSynchronous();
+		if (FilePickerClassLoaded)
+		{
+			FilePicker = NewObject<UInterchangeFilePickerBase>(GetTransientPackage(), FilePickerClassLoaded, NAME_None, RF_NoFlags);
+		}
+	}
+#endif
+	if(FilePicker)
+	{
+		FInterchangeFilePickerParameters Parameters;
+		Parameters.bAllowMultipleFiles = false;
+		Parameters.Title = FText::Format(NSLOCTEXT("Interchange", "ImportCustomLodAsync_FilePickerTitle", "Choose a file to import a custom LOD for LOD{0}"), FText::AsNumber(LodIndex));
+		TArray<FString> Filenames;
+		if (FilePicker->ScriptedFilePickerForTranslatorAssetType(EInterchangeTranslatorAssetType::Meshes, Parameters, Filenames))
+		{
+			//We set bAllowMultipleFile to false, we should have only one result
+			if (ensure(Filenames.Num() == 1))
+			{
+				const UInterchangeSourceData* SourceData = UInterchangeManager::GetInterchangeManager().CreateSourceData(Filenames[0]);
+				return InternalImportCustomLodAsync(Promise, MeshObject, LodIndex, SourceData);
+			}
+		}
+	}
+
+	Promise->SetValue(false);
+	return Promise->GetFuture();
+}
+
+TFuture<bool> UInterchangeMeshUtilities::ImportCustomLodAsync(UObject* MeshObject, const int32 LodIndex, const UInterchangeSourceData* SourceData)
+{
+	TSharedPtr<TPromise<bool>> Promise = MakeShared<TPromise<bool>>();
+	
+	return InternalImportCustomLodAsync(Promise, MeshObject, LodIndex, SourceData);
+}
+
+TFuture<bool> UInterchangeMeshUtilities::InternalImportCustomLodAsync(TSharedPtr<TPromise<bool>> Promise, UObject* MeshObject, const int32 LodIndex, const UInterchangeSourceData* SourceData)
+{
+#if WITH_EDITOR
+	UInterchangeManager& InterchangeManager = UInterchangeManager::GetInterchangeManager();
+
+	UInterchangeAssetImportData* InterchangeAssetImportData = nullptr;
+	USkeletalMesh* SkeletalMesh = Cast<USkeletalMesh>(MeshObject);
+	UStaticMesh* StaticMesh = Cast<UStaticMesh>(MeshObject);
+	EInterchangeReimportType ImportType = EInterchangeReimportType::AssetCustomLODImport;
+	if (SkeletalMesh)
+	{
+		InterchangeAssetImportData = Cast<UInterchangeAssetImportData>(SkeletalMesh->GetAssetImportData());
+		if (SkeletalMesh->GetLODNum() > LodIndex)
+		{
+			ImportType = EInterchangeReimportType::AssetCustomLODReimport;
+		}
+	}
+	else if (StaticMesh)
+	{
+		InterchangeAssetImportData = Cast<UInterchangeAssetImportData>(StaticMesh->GetAssetImportData());
+		if (StaticMesh->GetNumSourceModels() > LodIndex)
+		{
+			ImportType = EInterchangeReimportType::AssetCustomLODReimport;
+		}
+	}
+	else
+	{
+		//We support Import custom LOD only for skeletalmesh and staticmesh
+		Promise->SetValue(false);
+		return Promise->GetFuture();
+	}
+
+	FImportAssetParameters ImportAssetParameters;
+	ImportAssetParameters.bIsAutomated = true;
+	for (TObjectPtr<UInterchangePipelineBase> SelectedPipeline : InterchangeAssetImportData->Pipelines)
+	{
+		UInterchangePipelineBase* GeneratedPipeline = Cast<UInterchangePipelineBase>(StaticDuplicateObject(SelectedPipeline, GetTransientPackage()));
+		GeneratedPipeline->AdjustSettingsForReimportType(ImportType, nullptr);
+		ImportAssetParameters.OverridePipelines.Add(GeneratedPipeline);
+	}
+	FString ImportAssetPath = TEXT("/Engine/TempEditor/Interchange/") + FGuid::NewGuid().ToString(EGuidFormats::Base36Encoded);
+	UE::Interchange::FAssetImportResultRef AssetImportResult = InterchangeManager.ImportAssetAsync(ImportAssetPath, SourceData, ImportAssetParameters);
+	FString SourceDataFilename = SourceData->GetFilename();
+	if (SkeletalMesh)
+	{
+		AssetImportResult->OnDone([Promise, SkeletalMesh, LodIndex, SourceDataFilename](UE::Interchange::FImportResult& ImportResult)
+			{
+				USkeletalMesh* SourceSkeletalMesh = Cast< USkeletalMesh >(ImportResult.GetFirstAssetOfClass(USkeletalMesh::StaticClass()));
+
+				if(SourceSkeletalMesh)
+				{
+
+					Promise->SetValue(FLODUtilities::SetCustomLOD(SkeletalMesh, SourceSkeletalMesh, LodIndex, SourceDataFilename));
+				}
+				else
+				{
+					Promise->SetValue(false);
+				}
+				SourceSkeletalMesh->ClearFlags(RF_Standalone);
+				SourceSkeletalMesh->ClearInternalFlags(EInternalObjectFlags::Async);
+			});
+	}
+	else if (StaticMesh)
+	{
+		AssetImportResult->OnDone([Promise, StaticMesh, LodIndex, SourceDataFilename](UE::Interchange::FImportResult& ImportResult)
+			{
+				UStaticMesh* SourceStaticMesh = Cast< UStaticMesh >(ImportResult.GetFirstAssetOfClass(UStaticMesh::StaticClass()));
+				if(SourceStaticMesh)
+				{
+					Promise->SetValue(StaticMesh->SetCustomLOD(SourceStaticMesh, LodIndex, SourceDataFilename));
+				}
+				else
+				{
+					Promise->SetValue(false);
+				}
+				SourceStaticMesh->ClearFlags(RF_Standalone);
+				SourceStaticMesh->ClearInternalFlags(EInternalObjectFlags::Async);
+			});
+	}
+
+	return Promise->GetFuture();
+#else
+	Promise->SetValue(false);
+	return Promise->GetFuture();
+#endif
+}
\ No newline at end of file
diff --git a/Engine/Source/Runtime/Experimental/Interchange/Engine/Public/InterchangeFilePickerBase.h b/Engine/Source/Runtime/Experimental/Interchange/Engine/Public/InterchangeFilePickerBase.h
--- a/Engine/Source/Runtime/Experimental/Interchange/Engine/Public/InterchangeFilePickerBase.h
+++ b/Engine/Source/Runtime/Experimental/Interchange/Engine/Public/InterchangeFilePickerBase.h
@@ -9,6 +9,24 @@
 
 #include "InterchangeFilePickerBase.generated.h"
 
+USTRUCT(BlueprintType)
+struct FInterchangeFilePickerParameters
+{
+	GENERATED_BODY()
+
+	//If true the user will be able to select multiple files.
+	UPROPERTY(EditAnywhere, Category = "Interchange | File Picker")
+	bool bAllowMultipleFiles = false;
+
+	//If not empty it will override the default title
+	UPROPERTY(EditAnywhere, Category = "Interchange | File Picker")
+	FText Title = FText();
+
+	//Set the default open path that the dialog will show to the user
+	UPROPERTY(EditAnywhere, Category = "Interchange | File Picker")
+	FString DefaultPath = TEXT("");
+};
+
 UCLASS(Abstract, BlueprintType, Blueprintable, Experimental)
 class INTERCHANGEENGINE_API UInterchangeFilePickerBase : public UObject
 {
@@ -20,26 +38,26 @@ public:
 	 * Non virtual helper to allow blueprint to implement event base function to implement the function FilePickerForTranslatorAssetType,
 	 */
 	UFUNCTION(BlueprintNativeEvent, BlueprintCallable, Category = "Interchange | Utilities")
-	bool ScriptedFilePickerForTranslatorAssetType(const EInterchangeTranslatorAssetType TranslatorAssetType, FString& OutFilename);
+	bool ScriptedFilePickerForTranslatorAssetType(const EInterchangeTranslatorAssetType TranslatorAssetType, FInterchangeFilePickerParameters& Parameters, TArray<FString>& OutFilenames);
 
 	/** The default implementation (call if the blueprint do not have any implementation) will call the virtual FilePickerForTranslatorAssetType */
-	bool ScriptedFilePickerForTranslatorAssetType_Implementation(const EInterchangeTranslatorAssetType TranslatorAssetType, FString& OutFilename)
+	bool ScriptedFilePickerForTranslatorAssetType_Implementation(const EInterchangeTranslatorAssetType TranslatorAssetType, FInterchangeFilePickerParameters& Parameters, TArray<FString>& OutFilenames)
 	{
 		//By default we call the virtual function FilePickerForTranslatorAssetType
-		return FilePickerForTranslatorAssetType(TranslatorAssetType, OutFilename);
+		return FilePickerForTranslatorAssetType(TranslatorAssetType, Parameters, OutFilenames);
 	}
 
 	/**
 	 * Non virtual helper to allow blueprint to implement event base function to implement the function FilePickerForTranslatorType,
 	 */
 	UFUNCTION(BlueprintNativeEvent, BlueprintCallable, Category = "Interchange | Utilities")
-	bool ScriptedFilePickerForTranslatorType(const EInterchangeTranslatorType TranslatorType, FString& OutFilename);
+	bool ScriptedFilePickerForTranslatorType(const EInterchangeTranslatorType TranslatorType, FInterchangeFilePickerParameters& Parameters, TArray<FString>& OutFilenames);
 
 	/** The default implementation (call if the blueprint do not have any implementation) will call the virtual FilePickerForTranslatorType */
-	bool ScriptedFilePickerForTranslatorType_Implementation(const EInterchangeTranslatorType TranslatorAssetType, FString& OutFilename)
+	bool ScriptedFilePickerForTranslatorType_Implementation(const EInterchangeTranslatorType TranslatorAssetType, FInterchangeFilePickerParameters& Parameters, TArray<FString>& OutFilenames)
 	{
 		//By default we call the virtual function FilePickerForTranslatorType
-		return FilePickerForTranslatorType(TranslatorAssetType, OutFilename);
+		return FilePickerForTranslatorType(TranslatorAssetType, Parameters, OutFilenames);
 	}
 
 
@@ -51,7 +69,7 @@ protected:
 	 * 
 	 * @Return - true if it set OutFilename with a valid file path, false otherwise.
 	 */
-	virtual bool FilePickerForTranslatorAssetType(const EInterchangeTranslatorAssetType TranslatorAssetType, FString& OutFilename) { return false; }
+	virtual bool FilePickerForTranslatorAssetType(const EInterchangeTranslatorAssetType TranslatorAssetType, FInterchangeFilePickerParameters& Parameters, TArray<FString>& OutFilenames) { return false; }
 
 	/**
 	 * This function must set OutFilename with a valid file that can be import to create the specified TranslatorType.
@@ -59,6 +77,6 @@ protected:
 	 *
 	 * @Return - true if it set OutFilename with a valid file path, false otherwise.
 	 */
-	virtual bool FilePickerForTranslatorType(const EInterchangeTranslatorType TranslatorType, FString& OutFilename) { return false; }
+	virtual bool FilePickerForTranslatorType(const EInterchangeTranslatorType TranslatorType, FInterchangeFilePickerParameters& Parameters, TArray<FString>& OutFilenames) { return false; }
 
 };
diff --git a/Engine/Source/Runtime/Experimental/Interchange/Engine/Public/InterchangeMeshUtilities.h b/Engine/Source/Runtime/Experimental/Interchange/Engine/Public/InterchangeMeshUtilities.h
new file mode 100644
--- /dev/null
+++ b/Engine/Source/Runtime/Experimental/Interchange/Engine/Public/InterchangeMeshUtilities.h
@@ -0,0 +1,48 @@
+// Copyright Epic Games, Inc. All Rights Reserved.
+
+#pragma once
+
+#include "Async/Future.h"
+#include "UObject/Object.h"
+#include "UObject/ObjectMacros.h"
+
+#include "InterchangeMeshUtilities.generated.h"
+
+class UInterchangeSourceData;
+
+UCLASS(Experimental)
+class INTERCHANGEENGINE_API UInterchangeMeshUtilities : public UObject
+{
+	GENERATED_BODY()
+public:
+	/**
+	 * This function ask to select a file and import a mesh from it. After it add/replace the MeshObject LOD (at LodIndex) with the imported mesh LOD data.
+	 *
+	 * @Param MeshObject - The Mesh we want to add the lod
+	 * @Param LodIndex - The index of the LOD we want to replace or add
+	 * @Return - return true if it successfully add or replace the MeshObject LOD at LodIndex with the imported data.
+	 *
+	 * @Note - This function will sort out the SourceData and call ImportCustomLodAsync and wait until the result is available.
+	 */
+	static TFuture<bool> ImportCustomLodAsync(UObject* MeshObject, const int32 LodIndex);
+
+
+	/**
+	 * This function import a mesh from the source data and add/replace the MeshObject LOD (at LodIndex) with the imported mesh LOD data.
+	 *
+	 * @Param MeshObject - The Mesh we want to add the lod
+	 * @Param LodIndex - The index of the LOD we want to replace or add
+	 * @Param SourceData - The source to import the custom LOD
+	 * @Return - return future boolean which will be true if it successfully add or replace the MeshObject LOD at LodIndex with the imported data.
+	 *
+	 * @Note - Since this is the async version, the SourceData parameter is required.
+	 * @Note - This function will search for an available interchange asset factory that can do the job for the MeshObject class
+	 */
+	static TFuture<bool> ImportCustomLodAsync(UObject* MeshObject, const int32 LodIndex, const UInterchangeSourceData* SourceData);
+
+private:
+
+	static TFuture<bool> InternalImportCustomLodAsync(TSharedPtr<TPromise<bool>> Promise, UObject* MeshObject, const int32 LodIndex, const UInterchangeSourceData* SourceData);
+};
+
+
