Fix potential fatal error when saving a package with skeletal mesh still being loaded asynchronously

diff --git a/Engine/Source/Developer/SkeletalMeshUtilitiesCommon/Private/LODUtilities.cpp b/Engine/Source/Developer/SkeletalMeshUtilitiesCommon/Private/LODUtilities.cpp
--- a/Engine/Source/Developer/SkeletalMeshUtilitiesCommon/Private/LODUtilities.cpp
+++ b/Engine/Source/Developer/SkeletalMeshUtilitiesCommon/Private/LODUtilities.cpp
@@ -3221,6 +3221,9 @@ void FLODUtilities::BuildMorphTargets(USkeletalMesh* BaseSkelMesh, FSkeletalMesh
 			{
 				if (LODIndex == 0)
 				{
+					// Required both for NewObject and to avoid a fatal error in StaticFindObject.
+					FGCScopeGuard GCScopeGuard;
+
 					if (!IsInGameThread())
 					{
 						//TODO remove this code when overriding a UObject will be allow outside of the game thread
@@ -3236,7 +3239,7 @@ void FLODUtilities::BuildMorphTargets(USkeletalMesh* BaseSkelMesh, FSkeletalMesh
 						}
 						bNeedToClearAsyncFlag = true;
 					}
-					FGCScopeGuard GCScopeGuard;
+
 					MorphTarget = NewObject<UMorphTarget>(BaseSkelMesh, ObjectName);
 				}
 				else
diff --git a/Engine/Source/Runtime/CoreUObject/Private/UObject/SavePackage/SavePackageUtilities.cpp b/Engine/Source/Runtime/CoreUObject/Private/UObject/SavePackage/SavePackageUtilities.cpp
--- a/Engine/Source/Runtime/CoreUObject/Private/UObject/SavePackage/SavePackageUtilities.cpp
+++ b/Engine/Source/Runtime/CoreUObject/Private/UObject/SavePackage/SavePackageUtilities.cpp
@@ -185,7 +185,10 @@ namespace UE
 	{
 		if (InOuter == nullptr)
 		{
-			return GIsSavingPackage;
+			// That global is only meant to be set and read from the game-thread...
+			// otherwise it could interfere with normal operations coming from
+			// other threads like async loading, etc...
+			return GIsSavingPackage && IsInGameThread();
 		}
 		return InOuter->GetPackage()->HasAnyPackageFlags(PKG_IsSaving);
 	}
