PR #9527: Consider unreal's skeletal mesh bone naming when import non base lod skeletal mesh (Contributed by [REDACTED_PERSON])
#misc also applied fix to LODUtilities implementation

diff --git a/Engine/Source/Developer/SkeletalMeshUtilitiesCommon/Private/LODUtilities.cpp b/Engine/Source/Developer/SkeletalMeshUtilitiesCommon/Private/LODUtilities.cpp
--- a/Engine/Source/Developer/SkeletalMeshUtilitiesCommon/Private/LODUtilities.cpp
+++ b/Engine/Source/Developer/SkeletalMeshUtilitiesCommon/Private/LODUtilities.cpp
@@ -645,7 +645,7 @@ bool FLODUtilities::SetCustomLOD(USkeletalMesh* DestinationSkeletalMesh, USkelet
 			{
 				SkeletalMeshImportData::FBone& ImportedBone = LODImportData.RefBonesBinary[ImportBoneIndex];
 				int32 LODBoneIndex = ImportBoneIndex;
-				FName LODBoneName = FName(*ImportedBone.Name);
+				FName LODBoneName = FName(*FSkeletalMeshImportData::FixupBoneName(ImportedBone.Name));
 				int32 BaseBoneIndex = DestinationSkeletalMesh->GetRefSkeleton().FindBoneIndex(LODBoneName);
 				ImportDataBoneRemap[ImportBoneIndex] = BaseBoneIndex;
 				if (BaseBoneIndex != INDEX_NONE)
diff --git a/Engine/Source/Editor/UnrealEd/Private/Fbx/FbxSkeletalMeshImport.cpp b/Engine/Source/Editor/UnrealEd/Private/Fbx/FbxSkeletalMeshImport.cpp
--- a/Engine/Source/Editor/UnrealEd/Private/Fbx/FbxSkeletalMeshImport.cpp
+++ b/Engine/Source/Editor/UnrealEd/Private/Fbx/FbxSkeletalMeshImport.cpp
@@ -4161,7 +4161,7 @@ bool UnFbx::FFbxImporter::ImportSkeletalMeshLOD(USkeletalMesh* InSkeletalMesh, U
 			{
 				SkeletalMeshImportData::FBone& ImportedBone = LODImportData.RefBonesBinary[ImportBoneIndex];
 				int32 LODBoneIndex = ImportBoneIndex;
-				FName LODBoneName = FName(*ImportedBone.Name);
+				FName LODBoneName = FName(*FSkeletalMeshImportData::FixupBoneName(ImportedBone.Name));
 				int32 BaseBoneIndex = BaseSkeletalMesh->GetRefSkeleton().FindBoneIndex(LODBoneName);
 				ImportDataBoneRemap[ImportBoneIndex] = BaseBoneIndex;
 				if (BaseBoneIndex != INDEX_NONE)
