FLODUtilities::BuildMorphTargets
- Added checks to narrow down on a crash.
- Keep the async flag until we're done with the morph targets so that they don't get garbage collected.

diff --git a/Engine/Source/Developer/SkeletalMeshUtilitiesCommon/Private/LODUtilities.cpp b/Engine/Source/Developer/SkeletalMeshUtilitiesCommon/Private/LODUtilities.cpp
--- a/Engine/Source/Developer/SkeletalMeshUtilitiesCommon/Private/LODUtilities.cpp
+++ b/Engine/Source/Developer/SkeletalMeshUtilitiesCommon/Private/LODUtilities.cpp
@@ -3210,7 +3210,6 @@ void FLODUtilities::BuildMorphTargets(USkeletalMesh* BaseSkelMesh, FSkeletalMesh
 		FSkeletalMeshImportData& ShapeImportData = BaseImportData.MorphTargets[MorphTargetIndex];
 		TSet<uint32>& ModifiedPoints = BaseImportData.MorphTargetModifiedPoints[MorphTargetIndex];
 
-		bool bNeedToClearAsyncFlag = false;
 		UMorphTarget* MorphTarget = nullptr;
 		{
 			FName ObjectName = *ShapeName;
@@ -3237,7 +3236,6 @@ void FLODUtilities::BuildMorphTargets(USkeletalMesh* BaseSkelMesh, FSkeletalMesh
 							ExistingMorphTarget->Rename(nullptr, GetTransientPackage(), REN_ForceNoResetLoaders | REN_DoNotDirty | REN_DontCreateRedirectors | REN_NonTransactional);
 							ExistingMorphTarget = nullptr;
 						}
-						bNeedToClearAsyncFlag = true;
 					}
 
 					MorphTarget = NewObject<UMorphTarget>(BaseSkelMesh, ObjectName);
@@ -3253,6 +3251,7 @@ void FLODUtilities::BuildMorphTargets(USkeletalMesh* BaseSkelMesh, FSkeletalMesh
 
 		if (MorphTarget)
 		{
+			check(IsValid(MorphTarget));
 			MorphTargets.Add(MorphTarget);
 			int32 NewMorphDeltasIdx = Results.Add(new TArray< FMorphTargetDelta >());
 
@@ -3266,12 +3265,6 @@ void FLODUtilities::BuildMorphTargets(USkeletalMesh* BaseSkelMesh, FSkeletalMesh
 			NewWork->StartBackgroundTask(GLargeThreadPool);
 			CurrentNumTasks++;
 			NumTasks++;
-
-			if (bNeedToClearAsyncFlag)
-			{
-				const EInternalObjectFlags AsyncFlags = EInternalObjectFlags::Async | EInternalObjectFlags::AsyncLoading;
-				MorphTarget->ClearInternalFlags(AsyncFlags);
-			}
 		}
 
 		++ShapeIndex;
@@ -3307,6 +3300,7 @@ void FLODUtilities::BuildMorphTargets(USkeletalMesh* BaseSkelMesh, FSkeletalMesh
 		}
 
 		UMorphTarget* MorphTarget = MorphTargets[Index];
+		check(IsValid(MorphTarget));
 		MorphTarget->PopulateDeltas(*Results[Index], LODIndex, BaseLODModel.Sections, ShouldImportNormals == false, false, Thresholds.MorphThresholdPosition);
 
 		// register does mark package as dirty
