Fix truncation warnings in SkeletalMeshUtilitiesCommon & SkeletalMeshDescription.

#rnx

diff --git a/Engine/Source/Developer/SkeletalMeshUtilitiesCommon/Private/LODUtilities.cpp b/Engine/Source/Developer/SkeletalMeshUtilitiesCommon/Private/LODUtilities.cpp
--- a/Engine/Source/Developer/SkeletalMeshUtilitiesCommon/Private/LODUtilities.cpp
+++ b/Engine/Source/Developer/SkeletalMeshUtilitiesCommon/Private/LODUtilities.cpp
@@ -43,6 +43,8 @@
 #include "Framework/Notifications/NotificationManager.h"
 #include "Widgets/Notifications/SNotificationList.h"
 
+#include <limits>
+
 IMPLEMENT_MODULE(FDefaultModuleImpl, SkeletalMeshUtilitiesCommon)
 
 #define LOCTEXT_NAMESPACE "LODUtilities"
@@ -691,7 +693,8 @@ bool FLODUtilities::SetCustomLOD(USkeletalMesh* DestinationSkeletalMesh, USkelet
 			int32 LODBoneIndex = NewLODModel.ActiveBoneIndices[ActiveIndex];
 			FName LODBoneName = SourceSkeletalMesh->GetRefSkeleton().GetBoneName(LODBoneIndex);
 			int32 BaseBoneIndex = DestinationSkeletalMesh->GetRefSkeleton().FindBoneIndex(LODBoneName);
-			NewLODModel.ActiveBoneIndices[ActiveIndex] = BaseBoneIndex;
+			checkSlow(BaseBoneIndex != INDEX_NONE);
+			NewLODModel.ActiveBoneIndices[ActiveIndex] = static_cast<FBoneIndexType>(BaseBoneIndex);
 		}
 
 		// Fix up the chunk BoneMaps.
@@ -703,7 +706,8 @@ bool FLODUtilities::SetCustomLOD(USkeletalMesh* DestinationSkeletalMesh, USkelet
 				int32 LODBoneIndex = Section.BoneMap[BoneMapIndex];
 				FName LODBoneName = SourceSkeletalMesh->GetRefSkeleton().GetBoneName(LODBoneIndex);
 				int32 BaseBoneIndex = DestinationSkeletalMesh->GetRefSkeleton().FindBoneIndex(LODBoneName);
-				Section.BoneMap[BoneMapIndex] = BaseBoneIndex;
+				checkSlow(BaseBoneIndex != INDEX_NONE);
+				Section.BoneMap[BoneMapIndex] = static_cast<FBoneIndexType>(BaseBoneIndex);
 			}
 		}
 
@@ -714,7 +718,7 @@ bool FLODUtilities::SetCustomLOD(USkeletalMesh* DestinationSkeletalMesh, USkelet
 			int32 BaseBoneIndex = DestinationSkeletalMesh->GetRefSkeleton().FindBoneIndex(LODBoneName);
 			if (BaseBoneIndex != INDEX_NONE)
 			{
-				NewLODModel.RequiredBones[RequiredBoneIndex] = BaseBoneIndex;
+				NewLODModel.RequiredBones[RequiredBoneIndex] = static_cast<FBoneIndexType>(BaseBoneIndex);
 			}
 			else
 			{
@@ -906,7 +910,7 @@ bool VectorsOnSameSide(const FVector2D& Vec, const FVector2D& A, const FVector2D
 	return !(((B.Y - A.Y)*(Vec.X - A.X)) + ((A.X - B.X)*(Vec.Y - A.Y)) < 0.0f);
 }
 
-float PointToSegmentDistanceSquare(const FVector2D& A, const FVector2D& B, const FVector2D& P)
+double PointToSegmentDistanceSquare(const FVector2D& A, const FVector2D& B, const FVector2D& P)
 {
 	return FVector2D::DistSquared(P, FMath::ClosestPointOnSegment2D(P, A, B));
 }
@@ -949,7 +953,7 @@ bool VectorsOnSameSide(const FVector3f& Vec, const FVector3f& A, const FVector3f
 {
 	const FVector CrossA = FVector(Vec ^ A);
 	const FVector CrossB = FVector(Vec ^ B);
-	float DotWithEpsilon = SameSideDotProductEpsilon + (CrossA | CrossB);
+	double DotWithEpsilon = SameSideDotProductEpsilon + (CrossA | CrossB);
 	return !(DotWithEpsilon < 0.0f);
 }
 
@@ -1100,7 +1104,7 @@ void ProjectTargetOnBase(const TArray<FSoftSkinVertex>& BaseVertices, const TArr
 			QuadTree.Insert(TriangleElement.TriangleIndex, TriangleElement.UVsBound, DebugContext);
 		}
 		//Retrieve all triangle that are close to our point, let get 5% of UV extend
-		float DistanceThreshold = BaseMeshUVBound.GetExtent().Size()*0.05f;
+		double DistanceThreshold = BaseMeshUVBound.GetExtent().Size()*0.05;
 		//Find a match triangle for every target vertices
 		TArray<uint32> QuadTreeTriangleResults;
 		QuadTreeTriangleResults.Reserve(Triangles.Num() / 10); //Reserve 10% to speed up the query
@@ -1132,7 +1136,7 @@ void ProjectTargetOnBase(const TArray<FSoftSkinVertex>& BaseVertices, const TArr
 				CurBox = FBox2D(TargetUV - Extent, TargetUV + Extent);
 			}
 
-			auto GetDistancePointToBaseTriangle = [&Triangles, &TargetVertices, &TargetVertexIndex](const uint32 BaseTriangleIndex)->float
+			auto GetDistancePointToBaseTriangle = [&Triangles, &TargetVertices, &TargetVertexIndex](const uint32 BaseTriangleIndex)->double
 			{
 				FTriangleElement& CandidateTriangle = Triangles[BaseTriangleIndex];
 				return FVector::DistSquared(FMath::ClosestPointOnTriangleToPoint((FVector)TargetVertices[TargetVertexIndex].Position, (FVector)CandidateTriangle.Vertices[0].Position, (FVector)CandidateTriangle.Vertices[1].Position, (FVector)CandidateTriangle.Vertices[2].Position), (FVector)TargetVertices[TargetVertexIndex].Position);
@@ -1141,10 +1145,10 @@ void ProjectTargetOnBase(const TArray<FSoftSkinVertex>& BaseVertices, const TArr
 			auto FailSafeUnmatchVertex = [&GetDistancePointToBaseTriangle, &QuadTreeTriangleResults](uint32 &OutIndexMatch)->bool
 			{
 				bool bFoundMatch = false;
-				float ClosestTriangleDistSquared = MAX_flt;
+				double ClosestTriangleDistSquared = std::numeric_limits<double>::max();
 				for (uint32 MatchTriangleIndex : QuadTreeTriangleResults)
 				{
-					float TriangleDistSquared = GetDistancePointToBaseTriangle(MatchTriangleIndex);
+					double TriangleDistSquared = GetDistancePointToBaseTriangle(MatchTriangleIndex);
 					if (TriangleDistSquared < ClosestTriangleDistSquared)
 					{
 						ClosestTriangleDistSquared = TriangleDistSquared;
@@ -1173,7 +1177,7 @@ void ProjectTargetOnBase(const TArray<FSoftSkinVertex>& BaseVertices, const TArr
 						continue;
 					}
 				}
-				float ClosestTriangleDistSquared = MAX_flt;
+				double ClosestTriangleDistSquared = std::numeric_limits<double>::max();
 				if (MatchTriangleIndexes.Num() == 1)
 				{
 					//One match, this mean no mirror UVs simply take the single match
@@ -1185,7 +1189,7 @@ void ProjectTargetOnBase(const TArray<FSoftSkinVertex>& BaseVertices, const TArr
 					//Geometry can use mirror so the UVs are not unique. Use the closest match triangle to the point to find the best match
 					for (uint32 MatchTriangleIndex : MatchTriangleIndexes)
 					{
-						float TriangleDistSquared = GetDistancePointToBaseTriangle(MatchTriangleIndex);
+						double TriangleDistSquared = GetDistancePointToBaseTriangle(MatchTriangleIndex);
 						if (TriangleDistSquared < ClosestTriangleDistSquared)
 						{
 							ClosestTriangleDistSquared = TriangleDistSquared;
@@ -1992,7 +1996,7 @@ void MatchVertexIndexUsingPosition(
 	}
 
 	//Retrieve all triangles that are close to our point, start at 0.25% of OcTree extend
-	float DistanceThreshold = BaseMeshPositionBound.GetExtent().Size()*0.0025f;
+	double DistanceThreshold = BaseMeshPositionBound.GetExtent().Size()*0.0025;
 
 	//Find a match triangle for every target vertices
 	TArray<FTriangleElement> OcTreeTriangleResults;
@@ -2053,7 +2057,7 @@ void MatchVertexIndexUsingPosition(
 		}
 
 		//Get the 3D distance between a point and a destination triangle
-		auto GetDistanceSrcPointToDestTriangle = [&TrianglesDest, &PositionSrc](const uint32 DestTriangleIndex)->float
+		auto GetDistanceSrcPointToDestTriangle = [&TrianglesDest, &PositionSrc](const uint32 DestTriangleIndex)->double
 		{
 			FTriangleElement& CandidateTriangle = TrianglesDest[DestTriangleIndex];
 			return FVector::DistSquared(FMath::ClosestPointOnTriangleToPoint((FVector)PositionSrc, (FVector)CandidateTriangle.Vertices[0].Position, (FVector)CandidateTriangle.Vertices[1].Position, (FVector)CandidateTriangle.Vertices[2].Position), (FVector)PositionSrc);
@@ -2063,11 +2067,11 @@ void MatchVertexIndexUsingPosition(
 		auto FailSafeUnmatchVertex = [&GetDistanceSrcPointToDestTriangle, &OcTreeTriangleResults](uint32 &OutIndexMatch)->bool
 		{
 			bool bFoundMatch = false;
-			float ClosestTriangleDistSquared = MAX_flt;
+			double ClosestTriangleDistSquared = std::numeric_limits<double>::max();
 			for (const FTriangleElement& MatchTriangle : OcTreeTriangleResults)
 			{
-				int32 MatchTriangleIndex = MatchTriangle.TriangleIndex;
-				float TriangleDistSquared = GetDistanceSrcPointToDestTriangle(MatchTriangleIndex);
+				const int32 MatchTriangleIndex = MatchTriangle.TriangleIndex;
+				const double TriangleDistSquared = GetDistanceSrcPointToDestTriangle(MatchTriangleIndex);
 				if (TriangleDistSquared < ClosestTriangleDistSquared)
 				{
 					ClosestTriangleDistSquared = TriangleDistSquared;
@@ -2097,7 +2101,7 @@ void MatchVertexIndexUsingPosition(
 					continue;
 				}
 			}
-			float ClosestTriangleDistSquared = MAX_flt;
+			double ClosestTriangleDistSquared = std::numeric_limits<double>::max();
 			if (MatchTriangleIndexes.Num() == 1)
 			{
 				//One match, this mean no mirror UVs simply take the single match
@@ -2109,7 +2113,7 @@ void MatchVertexIndexUsingPosition(
 				//Geometry can use mirror so the UVs are not unique. Use the closest match triangle to the point to find the best match
 				for (uint32 MatchTriangleIndex : MatchTriangleIndexes)
 				{
-					float TriangleDistSquared = GetDistanceSrcPointToDestTriangle(MatchTriangleIndex);
+					double TriangleDistSquared = GetDistanceSrcPointToDestTriangle(MatchTriangleIndex);
 					if (TriangleDistSquared < ClosestTriangleDistSquared)
 					{
 						ClosestTriangleDistSquared = TriangleDistSquared;
@@ -2493,7 +2497,7 @@ bool FLODUtilities::UpdateAlternateSkinWeights(FSkeletalMeshLODModel& LODModelDe
 	{
 		const SkeletalMeshImportData::FRawBoneInfluence& RawInfluence = AlternateInfluences[InfluenceIndex];
 		SkeletalMeshImportData::FVertInfluence LODAlternateInfluence;
-		LODAlternateInfluence.BoneIndex = RawInfluence.BoneIndex;
+		LODAlternateInfluence.BoneIndex = static_cast<FBoneIndexType>(RawInfluence.BoneIndex);
 		LODAlternateInfluence.VertIndex = RawInfluence.VertexIndex;
 		LODAlternateInfluence.Weight = RawInfluence.Weight;
 		ImportedProfileData.SourceModelInfluences.Add(LODAlternateInfluence);
@@ -2647,7 +2651,7 @@ void FLODUtilities::GenerateImportedSkinWeightProfileData(FSkeletalMeshLODModel&
 					//Map to root of the section
 					BoneMapIndex = 0;
 				}
-				WeightForBone.Add(BoneMapIndex, VertInfluence.Weight);
+				WeightForBone.Add(static_cast<FBoneIndexType>(BoneMapIndex), VertInfluence.Weight);
 			}
 		}
 
@@ -3458,7 +3462,7 @@ void FLODUtilities::AdjustImportDataFaceMaterialIndex(const TArray<FSkeletalMate
 		//Update all the faces
 		for (int32 FaceIndex = 0; FaceIndex < LODFaces.Num(); ++FaceIndex)
 		{
-			LODFaces[FaceIndex].MeshMaterialIndex = MaterialRemap[LODFaces[FaceIndex].MeshMaterialIndex];
+			LODFaces[FaceIndex].MeshMaterialIndex = static_cast<uint16>(MaterialRemap[LODFaces[FaceIndex].MeshMaterialIndex]);
 		}
 	}
 }
@@ -3516,15 +3520,15 @@ namespace TriangleStripHelper
 			const FVector2D VectorA = SegmentEndA - SegmentStartA;
 			const FVector2D VectorB = SegmentEndB - SegmentStartB;
 
-			const float S = (-VectorA.Y * (SegmentStartA.X - SegmentStartB.X) + VectorA.X * (SegmentStartA.Y - SegmentStartB.Y)) / (-VectorB.X * VectorA.Y + VectorA.X * VectorB.Y);
-			const float T = (VectorB.X * (SegmentStartA.Y - SegmentStartB.Y) - VectorB.Y * (SegmentStartA.X - SegmentStartB.X)) / (-VectorB.X * VectorA.Y + VectorA.X * VectorB.Y);
+			const double S = (-VectorA.Y * (SegmentStartA.X - SegmentStartB.X) + VectorA.X * (SegmentStartA.Y - SegmentStartB.Y)) / (-VectorB.X * VectorA.Y + VectorA.X * VectorB.Y);
+			const double T = (VectorB.X * (SegmentStartA.Y - SegmentStartB.Y) - VectorB.Y * (SegmentStartA.X - SegmentStartB.X)) / (-VectorB.X * VectorA.Y + VectorA.X * VectorB.Y);
 
 			return (S >= 0 && S <= 1 && T >= 0 && T <= 1);
 		};
 
 		auto IsInsideTriangle = [&Triangle, &SegmentIntersection2D, &Box, &TriangleBox](const FVector2D& TestPoint)->bool
 		{
-			float Extent = (2.0f * Box.GetSize().Size()) + (2.0f * TriangleBox.GetSize().Size());
+			double Extent = (2.0 * Box.GetSize().Size()) + (2.0 * TriangleBox.GetSize().Size());
 			FVector2D TestPointExtend(Extent, Extent);
 			int32 IntersectionCount = SegmentIntersection2D(Triangle.Vertices[0], Triangle.Vertices[1], TestPoint, TestPoint + TestPointExtend) ? 1 : 0;
 			IntersectionCount += SegmentIntersection2D(Triangle.Vertices[1], Triangle.Vertices[2], TestPoint, TestPoint + TestPointExtend) ? 1 : 0;
@@ -3625,14 +3629,14 @@ bool FLODUtilities::StripLODGeometry(USkeletalMesh* SkeletalMesh, const int32 LO
 
 		auto ShouldStripTriangle = [&](const FVector2D& UvA, const FVector2D& UvB, const FVector2D& UvC)->bool
 		{
-			FVector2D PixelUvA = FVector2D(FMath::FloorToInt(UvA.X * (float)ResX) % (ResX + 1), FMath::FloorToInt(UvA.Y * (float)ResY) % (ResY + 1));
-			FVector2D PixelUvB = FVector2D(FMath::FloorToInt(UvB.X * (float)ResX) % (ResX + 1), FMath::FloorToInt(UvB.Y * (float)ResY) % (ResY + 1));
-			FVector2D PixelUvC = FVector2D(FMath::FloorToInt(UvC.X * (float)ResX) % (ResX + 1), FMath::FloorToInt(UvC.Y * (float)ResY) % (ResY + 1));
+			const FIntVector2 PixelUvA(FMath::FloorToInt32(UvA.X * ResX) % (ResX + 1), FMath::FloorToInt32(UvA.Y * ResY) % (ResY + 1));
+			const FIntVector2 PixelUvB(FMath::FloorToInt32(UvB.X * ResX) % (ResX + 1), FMath::FloorToInt32(UvB.Y * ResY) % (ResY + 1));
+			const FIntVector2 PixelUvC(FMath::FloorToInt32(UvC.X * ResX) % (ResX + 1), FMath::FloorToInt32(UvC.Y * ResY) % (ResY + 1));
 
-			int32 MinU = FMath::Clamp(FMath::Min3<int32>(PixelUvA.X, PixelUvB.X, PixelUvC.X), 0, ResX);
-			int32 MinV = FMath::Clamp(FMath::Min3<int32>(PixelUvA.Y, PixelUvB.Y, PixelUvC.Y), 0, ResY);
-			int32 MaxU = FMath::Clamp(FMath::Max3<int32>(PixelUvA.X, PixelUvB.X, PixelUvC.X), 0, ResX);
-			int32 MaxV = FMath::Clamp(FMath::Max3<int32>(PixelUvA.Y, PixelUvB.Y, PixelUvC.Y), 0, ResY);
+			const int32 MinU = FMath::Clamp(FMath::Min3<int32>(PixelUvA.X, PixelUvB.X, PixelUvC.X), 0, ResX);
+			const int32 MinV = FMath::Clamp(FMath::Min3<int32>(PixelUvA.Y, PixelUvB.Y, PixelUvC.Y), 0, ResY);
+			const int32 MaxU = FMath::Clamp(FMath::Max3<int32>(PixelUvA.X, PixelUvB.X, PixelUvC.X), 0, ResX);
+			const int32 MaxV = FMath::Clamp(FMath::Max3<int32>(PixelUvA.Y, PixelUvB.Y, PixelUvC.Y), 0, ResY);
 
 			//Do not read the alpha value when testing the texture value
 			auto IsPixelZero = [&](int32 PosX, int32 PosY) -> bool
@@ -3640,9 +3644,9 @@ bool FLODUtilities::StripLODGeometry(USkeletalMesh* SkeletalMesh, const int32 LO
 				const int32 RefPos = PosX + (PosY * InitialSource.GetSizeX());
 				const void * PixelPtr = Ref2DData.GetData() + RefPos * FormatDataSize;
 				
-				FLinearColor Color = ERawImageFormat::GetOnePixelLinear(PixelPtr,RawFormat,bSRGB);	
+				const FLinearColor Color = ERawImageFormat::GetOnePixelLinear(PixelPtr,RawFormat,bSRGB);	
 
-				bool bPixelIsZero = 
+				const bool bPixelIsZero = 
 					FMath::IsNearlyZero(Color.R,Threshold) &&
 					FMath::IsNearlyZero(Color.G,Threshold) &&
 					FMath::IsNearlyZero(Color.B,Threshold);
@@ -3660,19 +3664,19 @@ bool FLODUtilities::StripLODGeometry(USkeletalMesh* SkeletalMesh, const int32 LO
 			{
 				for (int32 PosX = MinU; PosX < MaxU; ++PosX)
 				{
-					bool bStripPixel = IsPixelZero(PosX, PosY);
+					const bool bStripPixel = IsPixelZero(PosX, PosY);
 
 					//if any none zeroed pixel intersect the triangle, prevent stripping of this triangle
 					if (!bStripPixel)
 					{
-						FVector2D StartPixel((float)PosX, (float)PosY);
-						FVector2D EndPixel((float)(PosX+1), (float)(PosY + 1));
+						FVector2D StartPixel(PosX, PosY);
+						FVector2D EndPixel(PosX + 1, PosY + 1);
 						FBox2D Box2D(StartPixel, EndPixel);
 						//Test if the triangle UV touch this pixel
 						TriangleStripHelper::FTriangle2D Triangle;
-						Triangle.Vertices[0] = PixelUvA;
-						Triangle.Vertices[1] = PixelUvB;
-						Triangle.Vertices[2] = PixelUvC;
+						Triangle.Vertices[0] = FVector2D(PixelUvA.X, PixelUvA.Y);
+						Triangle.Vertices[1] = FVector2D(PixelUvB.X, PixelUvB.Y);
+						Triangle.Vertices[2] = FVector2D(PixelUvC.X, PixelUvC.Y);
 						if(TriangleStripHelper::IntersectTriangleAndAABB(Triangle, Box2D))
 						{
 							return false;
diff --git a/Engine/Source/Runtime/SkeletalMeshDescription/Private/SkeletalMeshAttributes.cpp b/Engine/Source/Runtime/SkeletalMeshDescription/Private/SkeletalMeshAttributes.cpp
--- a/Engine/Source/Runtime/SkeletalMeshDescription/Private/SkeletalMeshAttributes.cpp
+++ b/Engine/Source/Runtime/SkeletalMeshDescription/Private/SkeletalMeshAttributes.cpp
@@ -18,7 +18,7 @@ static FString SkinWeightAttributeNamePrefix()
 
 void FSkeletalMeshAttributes::Register()
 {
-	MeshDescription.VertexAttributes().RegisterAttribute<int32[]>(MeshAttribute::Vertex::SkinWeights, 1, 0.0f, EMeshAttributeFlags::Mandatory);
+	MeshDescription.VertexAttributes().RegisterAttribute<int32[]>(MeshAttribute::Vertex::SkinWeights, 1, 0, EMeshAttributeFlags::Mandatory);
 
 	// Call super class
 	FStaticMeshAttributes::Register();
