LODUtilities: Fixed potential dereferenting null pointers in SetCustomLOD.

diff --git a/Engine/Source/Developer/SkeletalMeshUtilitiesCommon/Private/LODUtilities.cpp b/Engine/Source/Developer/SkeletalMeshUtilitiesCommon/Private/LODUtilities.cpp
--- a/Engine/Source/Developer/SkeletalMeshUtilitiesCommon/Private/LODUtilities.cpp
+++ b/Engine/Source/Developer/SkeletalMeshUtilitiesCommon/Private/LODUtilities.cpp
@@ -538,10 +538,15 @@ bool FLODUtilities::SetCustomLOD(USkeletalMesh* DestinationSkeletalMesh, USkelet
 	TArray<int32> ClothingAssetSectionIndices;
 	TArray<int32> ClothingAssetInternalLodIndices;
 
-	FSkeletalMeshModel* SourceImportedResource = SourceSkeletalMesh->GetImportedModel();
-	FSkeletalMeshModel* DestImportedResource = DestinationSkeletalMesh->GetImportedModel();
+	FSkeletalMeshModel* const SourceImportedResource = SourceSkeletalMesh->GetImportedModel();
+	FSkeletalMeshModel* const DestImportedResource = DestinationSkeletalMesh->GetImportedModel();
 
-	if (SourceImportedResource && SourceImportedResource->LODModels.IsValidIndex(LodIndex))
+	if (!SourceImportedResource || !DestImportedResource)
+	{
+		return false;
+	}
+
+	if (SourceImportedResource->LODModels.IsValidIndex(LodIndex))
 	{
 		bMustReimportAlternateSkinWeightProfile = true;
 		FLODUtilities::UnbindClothingAndBackup(DestinationSkeletalMesh, ClothingBindings, LodIndex);
@@ -550,7 +555,7 @@ bool FLODUtilities::SetCustomLOD(USkeletalMesh* DestinationSkeletalMesh, USkelet
 	//Lambda to call to re-apply the clothing
 	auto ReapplyClothing = [&DestinationSkeletalMesh, &ClothingBindings, &SourceImportedResource, &LodIndex]()
 	{
-		if (SourceImportedResource && SourceImportedResource->LODModels.IsValidIndex(LodIndex))
+		if (SourceImportedResource->LODModels.IsValidIndex(LodIndex))
 		{
 			// Re-apply our clothing assets
 			FLODUtilities::RestoreClothingFromBackup(DestinationSkeletalMesh, ClothingBindings, LodIndex);
