Add a utility function to convert static meshes to skeletal meshes using a reference skeleton and a skeleton object.
Add an RMB entry to the static mesh asset browser menu to initiate conversion.
Fix computation of the bounding sphere of a skelmesh to match what static mesh and TBoxSphereBounds actually do. Previously the bounding sphere would be contained by the longest side of the bounding box, whereas in fact the sphere should contain the entirety of the box.

diff --git a/Engine/Plugins/Experimental/Animation/SkeletalMeshModelingTools/SkeletalMeshModelingTools.uplugin b/Engine/Plugins/Experimental/Animation/SkeletalMeshModelingTools/SkeletalMeshModelingTools.uplugin
--- a/Engine/Plugins/Experimental/Animation/SkeletalMeshModelingTools/SkeletalMeshModelingTools.uplugin
+++ b/Engine/Plugins/Experimental/Animation/SkeletalMeshModelingTools/SkeletalMeshModelingTools.uplugin
@@ -2,7 +2,7 @@
 	"FileVersion": 3,
 	"Version": 1,
 	"VersionName": "0.1",
-	"FriendlyName": "SkeletalMeshModelingTools",
+	"FriendlyName": "Skeletal Mesh Modeling Tools",
 	"Description": "Modeling Tools for the Skeletal Mesh Editor",
 	"Category" : "Animation",
 	"CreatedBy" : "Epic Games, Inc.",
@@ -22,10 +22,6 @@
 		}
 	],
 	"Plugins": [
-		{
-			"Name": "ComputeFramework",
-			"Enabled": true
-		},
 		{
 			"Name": "MeshModelingToolsetExp",
 			"Enabled": true
diff --git a/Engine/Plugins/Experimental/Animation/SkeletalMeshModelingTools/Source/SkeletalMeshModelingTools/Private/SkeletalMeshModelingToolsMeshConverter.cpp b/Engine/Plugins/Experimental/Animation/SkeletalMeshModelingTools/Source/SkeletalMeshModelingTools/Private/SkeletalMeshModelingToolsMeshConverter.cpp
new file mode 100644
--- /dev/null
+++ b/Engine/Plugins/Experimental/Animation/SkeletalMeshModelingTools/Source/SkeletalMeshModelingTools/Private/SkeletalMeshModelingToolsMeshConverter.cpp
@@ -0,0 +1,233 @@
+﻿// Copyright Epic Games, Inc. All Rights Reserved.
+
+#include "SkeletalMeshModelingToolsMeshConverter.h"
+
+#include "AssetToolsModule.h"
+#include "Animation/Skeleton.h"
+#include "Engine/SkeletalMesh.h"
+#include "Engine/StaticMesh.h"
+#include "PropertyEditorModule.h"
+#include "StaticToSkeletalMeshConverter.h"
+#include "Dialog/SCustomDialog.h"
+
+DEFINE_LOG_CATEGORY_STATIC(LogSkeletalMeshModelingToolsMeshConverter, Log, All)
+
+#define LOCTEXT_NAMESPACE "SkeletalMeshModelingToolsMeshConverter"
+
+
+USkeletonFromStaticMeshFactory::USkeletonFromStaticMeshFactory(
+	const FObjectInitializer& InObjectInitializer
+	) :
+	Super(InObjectInitializer)
+{
+	SupportedClass = USkeleton::StaticClass();
+	bCreateNew = true;
+	bEditAfterNew = false;
+}
+
+
+UObject* USkeletonFromStaticMeshFactory::FactoryCreateNew(
+	UClass* InClass,
+	UObject* InParent, 
+	FName InName,
+	EObjectFlags InFlags,
+	UObject* InContext,
+	FFeedbackContext* InWarn
+	)
+{
+	static const FVector RootBoneRelativePosition(0.5, 0.5, 0.0);
+	return FStaticToSkeletalMeshConverter::CreateSkeletonFromStaticMesh(InParent, InName, InFlags, StaticMesh, RootBoneRelativePosition);
+}
+
+
+USkeletalMeshFromStaticMeshFactory::USkeletalMeshFromStaticMeshFactory(
+	const FObjectInitializer& InObjectInitializer
+	) :
+	Super(InObjectInitializer)
+{
+	SupportedClass = USkeletalMesh::StaticClass();
+	bCreateNew = true;
+	bEditAfterNew = true;
+}
+
+
+UObject* USkeletalMeshFromStaticMeshFactory::FactoryCreateNew(
+	UClass* InClass,
+	UObject* InParent,
+	FName InName,
+	EObjectFlags InFlags,
+	UObject* InContext,
+	FFeedbackContext* InWarn
+	)
+{
+	USkeletalMesh* SkeletalMesh = FStaticToSkeletalMeshConverter::CreateSkeletalMeshFromStaticMesh(
+		InParent, InName, InFlags, StaticMesh, ReferenceSkeleton, BindBoneName);
+
+	// Update the skeletal mesh and the skeleton so that their ref skeletons are in sync and the skeleton's preview mesh
+	// is the one we just created.
+	SkeletalMesh->SetSkeleton(Skeleton);
+	Skeleton->SetPreviewMesh(SkeletalMesh);
+	Skeleton->MergeAllBonesToBoneTree(SkeletalMesh);
+	
+	return SkeletalMesh;
+}
+
+
+USkeleton* UStaticMeshToSkeletalMeshConvertOptions::GetSkeleton(
+	bool& bInvalidSkeletonIsError,
+	const IPropertyHandle* PropertyHandle)
+{
+	// bInvalidSkeletonIsError = true;
+
+	switch (SkeletonImportOption)
+	{
+	default:
+	case EReferenceSkeletonImportOption::CreateNew:
+		return nullptr;
+	case EReferenceSkeletonImportOption::UseExistingSkeleton:
+		return Cast<USkeleton>(Skeleton.ResolveObject());
+	case EReferenceSkeletonImportOption::UseExistingSkeletalMesh:
+		{
+			USkeletalMesh* SkeletalMeshPtr = Cast<USkeletalMesh>(SkeletalMesh.ResolveObject());
+			return SkeletalMeshPtr ? SkeletalMeshPtr->GetSkeleton() : nullptr;
+		}
+	}
+}
+
+static bool ConvertSingleMeshToSkeletalMeshInteractive(
+	UStaticMesh* InStaticMesh,
+	TArray<UObject*>& OutObjectsAdded
+	)
+{
+	UStaticMeshToSkeletalMeshConvertOptions* Options = NewObject<UStaticMeshToSkeletalMeshConvertOptions>();
+	FPropertyEditorModule& PropertyEditorModule = FModuleManager::GetModuleChecked<FPropertyEditorModule>("PropertyEditor");
+
+	FDetailsViewArgs DetailsViewArgs;
+	DetailsViewArgs.bAllowSearch = false;
+	DetailsViewArgs.NameAreaSettings = FDetailsViewArgs::HideNameArea;
+	TSharedRef<IDetailsView> DetailsView = PropertyEditorModule.CreateDetailView(DetailsViewArgs);
+	DetailsView->SetObject(Options);
+
+	TSharedRef<SCustomDialog> OptionsDialog = SNew(SCustomDialog)
+		.Title(LOCTEXT("OptionsDialogTitle", "Static Mesh Conversion Options"))
+		.Content()
+		[
+			DetailsView
+		]
+	.Buttons({
+		SCustomDialog::FButton(LOCTEXT("DialogButtonConvert", "Convert")),
+		SCustomDialog::FButton(LOCTEXT("DialogButtonCancel", "Cancel"))
+	});
+
+	if (OptionsDialog->ShowModal() != 0)
+	{
+		return false;
+	}
+
+	Options->SaveConfig();
+
+	// Make these configurable.
+	// ReSharper disable CppTooWideScope
+	static const TCHAR* DefaultSkeletonSuffix = TEXT("_Skel");
+	static const TCHAR* DefaultSkeletalMeshSuffix = TEXT("_SkelMesh");
+
+	IAssetTools& AssetTools = FAssetToolsModule::GetModule().Get();
+	FReferenceSkeleton ReferenceSkeleton;
+	USkeleton* Skeleton = nullptr;
+
+	switch(Options->SkeletonImportOption)
+	{
+	case EReferenceSkeletonImportOption::CreateNew:
+		{
+			FString SkeletonName;
+			FString SkeletonPackageName;
+			AssetTools.CreateUniqueAssetName(InStaticMesh->GetOutermost()->GetName(), DefaultSkeletonSuffix, SkeletonPackageName, SkeletonName);
+
+			USkeletonFromStaticMeshFactory* SkeletonFactory = NewObject<USkeletonFromStaticMeshFactory>();
+			SkeletonFactory->StaticMesh = InStaticMesh;
+
+			Skeleton = Cast<USkeleton>(AssetTools.CreateAsset(SkeletonName, FPackageName::GetLongPackagePath(SkeletonPackageName), USkeleton::StaticClass(), SkeletonFactory));
+			ReferenceSkeleton = Skeleton->GetReferenceSkeleton();
+		}
+		break;
+		
+	case EReferenceSkeletonImportOption::UseExistingSkeleton:
+		{
+			Skeleton = Cast<USkeleton>(Options->Skeleton.ResolveObject());
+			if (!Skeleton)
+			{
+				UE_LOG(LogSkeletalMeshModelingToolsMeshConverter, Error, TEXT("No valid skeleton given as a reference. Aborting conversion."));
+				return false;
+			}
+			ReferenceSkeleton = Skeleton->GetReferenceSkeleton();
+		}
+		break;
+	case EReferenceSkeletonImportOption::UseExistingSkeletalMesh:
+		{
+			USkeletalMesh* SkeletalMesh = Cast<USkeletalMesh>(Options->SkeletalMesh.ResolveObject());
+			if (!SkeletalMesh)
+			{
+				UE_LOG(LogSkeletalMeshModelingToolsMeshConverter, Error, TEXT("No valid skeletal mesh given as reference. Aborting conversion."));
+				return false;
+			}
+			Skeleton = SkeletalMesh->GetSkeleton();
+			ReferenceSkeleton = SkeletalMesh->GetRefSkeleton();
+		}
+		break;
+	}
+
+	USkeletalMesh *SkeletalMesh = nullptr;
+	if (ensure(Skeleton) && ensure(ReferenceSkeleton.GetNum() != 0))
+	{
+		FString SkeletalMeshName;
+		FString SkeletalMeshPackageName;
+		AssetTools.CreateUniqueAssetName(InStaticMesh->GetOutermost()->GetName(), DefaultSkeletalMeshSuffix, SkeletalMeshPackageName, SkeletalMeshName);
+
+		USkeletalMeshFromStaticMeshFactory* SkeletalMeshFactory = NewObject<USkeletalMeshFromStaticMeshFactory>();
+		SkeletalMeshFactory->StaticMesh = InStaticMesh;
+		SkeletalMeshFactory->ReferenceSkeleton = ReferenceSkeleton;
+		SkeletalMeshFactory->Skeleton = Skeleton;
+		SkeletalMeshFactory->BindBoneName = Options->BindingBoneName.BoneName; 
+
+		SkeletalMesh = Cast<USkeletalMesh>(AssetTools.CreateAsset(SkeletalMeshName, FPackageName::GetLongPackagePath(SkeletalMeshPackageName), USkeletalMesh::StaticClass(), SkeletalMeshFactory));
+	}
+
+	if (ensure(SkeletalMesh))
+	{
+		if (Options->SkeletonImportOption == EReferenceSkeletonImportOption::CreateNew)
+		{
+			OutObjectsAdded.Add(Skeleton);
+		}
+		OutObjectsAdded.Add(SkeletalMesh);
+	}
+
+	return true;
+}
+
+
+void ConvertStaticMeshToSkeletalMeshInteractive(
+	const TArray<FAssetData>& InStaticMeshAssets
+	)
+{
+	IAssetTools& AssetTools = FAssetToolsModule::GetModule().Get();
+	
+	TArray<UObject*> ObjectsAdded;
+
+	for (const FAssetData& AssetData: InStaticMeshAssets)
+	{
+		UStaticMesh* StaticMesh = Cast<UStaticMesh>(AssetData.GetAsset());
+		if (!ensure(StaticMesh))
+		{
+			continue;
+		}
+
+		if (!ConvertSingleMeshToSkeletalMeshInteractive(StaticMesh, ObjectsAdded))
+		{
+			break;
+		}
+	}
+	
+	AssetTools.SyncBrowserToAssets(ObjectsAdded);
+}
+
+#undef LOCTEXT_NAMESPACE
diff --git a/Engine/Plugins/Experimental/Animation/SkeletalMeshModelingTools/Source/SkeletalMeshModelingTools/Private/SkeletalMeshModelingToolsMeshConverter.h b/Engine/Plugins/Experimental/Animation/SkeletalMeshModelingTools/Source/SkeletalMeshModelingTools/Private/SkeletalMeshModelingToolsMeshConverter.h
new file mode 100644
--- /dev/null
+++ b/Engine/Plugins/Experimental/Animation/SkeletalMeshModelingTools/Source/SkeletalMeshModelingTools/Private/SkeletalMeshModelingToolsMeshConverter.h
@@ -0,0 +1,111 @@
+﻿// Copyright Epic Games, Inc. All Rights Reserved.
+
+#pragma once
+
+#include "BoneContainer.h"
+#include "Factories/Factory.h"
+#include "ReferenceSkeleton.h"
+#include "Interfaces/Interface_BoneReferenceSkeletonProvider.h"
+
+#include "SkeletalMeshModelingToolsMeshConverter.generated.h"
+
+
+class USkeleton;
+class UStaticMesh;
+struct FAssetData;
+
+
+UCLASS(Hidden, MinimalAPI)
+class USkeletonFromStaticMeshFactory :
+	public UFactory
+{
+	GENERATED_BODY()
+public:
+	USkeletonFromStaticMeshFactory(const FObjectInitializer& InObjectInitializer);
+
+	// UFactory overrides
+	virtual UObject* FactoryCreateNew(
+		UClass* InClass,
+		UObject* InParent,
+		FName InName,
+		EObjectFlags InFlags,
+		UObject* InContext,
+		FFeedbackContext* InWarn
+		) override;
+
+	virtual bool ShouldShowInNewMenu() const override { return false; }
+	
+	UPROPERTY()
+	TObjectPtr<UStaticMesh> StaticMesh;
+};
+
+
+UCLASS(Hidden, MinimalAPI)
+class USkeletalMeshFromStaticMeshFactory :
+	public UFactory
+{
+	GENERATED_BODY()
+public:
+	USkeletalMeshFromStaticMeshFactory(const FObjectInitializer& InObjectInitializer);
+
+	// UFactory overrides
+	virtual UObject* FactoryCreateNew(
+		UClass* InClass,
+		UObject* InParent,
+		FName InName,
+		EObjectFlags InFlags,
+		UObject* InContext,
+		FFeedbackContext* InWarn
+		) override;
+	
+	virtual bool ShouldShowInNewMenu() const override { return false; }
+
+	FReferenceSkeleton ReferenceSkeleton;
+
+	UPROPERTY()
+	TObjectPtr<USkeleton> Skeleton;
+	
+	UPROPERTY()
+	TObjectPtr<UStaticMesh> StaticMesh;
+
+	UPROPERTY()
+	FName BindBoneName;
+};
+
+
+UENUM()
+enum class EReferenceSkeletonImportOption
+{
+	CreateNew,
+	UseExistingSkeleton,
+	UseExistingSkeletalMesh
+};
+
+
+UCLASS(config=EditorPerProjectUserSettings, HideCategories=Object, MinimalAPI)
+class UStaticMeshToSkeletalMeshConvertOptions :
+	public UObject,
+	public IBoneReferenceSkeletonProvider
+{
+	GENERATED_BODY()
+public:
+
+	UPROPERTY(EditAnywhere, Config, Category="Import")
+	EReferenceSkeletonImportOption SkeletonImportOption = EReferenceSkeletonImportOption::CreateNew;
+
+	UPROPERTY(EditAnywhere, Config, Category="Import", Meta=(AllowedClasses="/Script/Engine.Skeleton", EditCondition="ImportOption==EReferenceSkeletonImportOption::UseExistingSkeleton", EditConditionHides))
+	FSoftObjectPath Skeleton;
+
+	UPROPERTY(EditAnywhere, Config, Category="Import", Meta=(AllowedClasses="/Script/Engine.SkeletalMesh", EditCondition="ImportOption==EReferenceSkeletonImportOption::UseExistingSkeletalMesh", EditConditionHides))
+	FSoftObjectPath SkeletalMesh;
+
+	UPROPERTY(EditAnywhere, Config, Category="Import", meta=(EditCondition="ImportOption!=EReferenceSkeletonImportOption::CreateNew", EditConditionHides))
+	FBoneReference BindingBoneName;
+
+	// IBoneReferenceSkeletonProvider implementation.
+	virtual USkeleton* GetSkeleton(bool& bInvalidSkeletonIsError, const IPropertyHandle* PropertyHandle) override;
+};
+
+
+/** Converts a static mesh to a skeletal mesh, creating both a skeletal mesh and a skeleton assets that link up. */
+void ConvertStaticMeshToSkeletalMeshInteractive(const TArray<FAssetData>& InStaticMeshAssets);
diff --git a/Engine/Plugins/Experimental/Animation/SkeletalMeshModelingTools/Source/SkeletalMeshModelingTools/Private/SkeletalMeshModelingToolsModule.cpp b/Engine/Plugins/Experimental/Animation/SkeletalMeshModelingTools/Source/SkeletalMeshModelingTools/Private/SkeletalMeshModelingToolsModule.cpp
--- a/Engine/Plugins/Experimental/Animation/SkeletalMeshModelingTools/Source/SkeletalMeshModelingTools/Private/SkeletalMeshModelingToolsModule.cpp
+++ b/Engine/Plugins/Experimental/Animation/SkeletalMeshModelingTools/Source/SkeletalMeshModelingTools/Private/SkeletalMeshModelingToolsModule.cpp
@@ -4,8 +4,10 @@
 
 #include "SkeletalMeshModelingToolsCommands.h"
 #include "SkeletalMeshModelingToolsEditorMode.h"
+#include "SkeletalMeshModelingToolsMeshConverter.h"
 #include "SkeletalMeshModelingToolsStyle.h"
 
+#include "ContentBrowserMenuContexts.h"
 #include "EditorModeManager.h"
 #include "EditorModeRegistry.h"
 #include "ISkeletalMeshEditor.h"
@@ -13,6 +15,7 @@
 #include "Modules/ModuleManager.h"
 #include "SkeletalMeshToolMenuContext.h"
 #include "ToolMenus.h"
+#include "Styling/SlateIconFinder.h"
 #include "WorkflowOrientedApp/ApplicationMode.h"
 
 
@@ -59,11 +62,11 @@ void FSkeletalMeshModelingToolsModule::RegisterMenusAndToolbars()
 {
 	FToolMenuOwnerScoped OwnerScoped(this);
 
-	UToolMenu* Toolbar = UToolMenus::Get()->ExtendMenu("AssetEditor.SkeletalMeshEditor.ToolBar");
+	if (UToolMenu* ToolMenu = UToolMenus::Get()->ExtendMenu("AssetEditor.SkeletalMeshEditor.ToolBar"))
 	{
-		FToolMenuSection& Section = Toolbar->FindOrAddSection("SkeletalMesh");
+		FToolMenuSection& Section = ToolMenu->FindOrAddSection("SkeletalMesh");
 		Section.AddDynamicEntry("ToggleModelingToolsMode", FNewToolMenuSectionDelegate::CreateLambda([this](FToolMenuSection& InSection) {
-			USkeletalMeshToolMenuContext* Context = InSection.FindContext<USkeletalMeshToolMenuContext>();
+			const USkeletalMeshToolMenuContext* Context = InSection.FindContext<USkeletalMeshToolMenuContext>();
 			if (Context && Context->SkeletalMeshEditor.IsValid())
 			{
 				InSection.AddEntry(FToolMenuEntry::InitToolBarButton(
@@ -74,6 +77,28 @@ void FSkeletalMeshModelingToolsModule::RegisterMenusAndToolbars()
 			}
 		}));
 	}
+
+	// Extend the asset browser for static meshes to include the conversion to skelmesh.
+	if (UToolMenu* ToolMenu = UToolMenus::Get()->ExtendMenu("ContentBrowser.AssetContextMenu.StaticMesh"))
+	{
+		FToolMenuSection& Section = ToolMenu->FindOrAddSection("GetAssetActions");
+		Section.AddDynamicEntry("ConvertToSkeletalMesh", FNewToolMenuSectionDelegate::CreateLambda([](FToolMenuSection& InSection)
+		{
+			const TAttribute<FText> Label = LOCTEXT("ConvertToSkeletalMesh", "Convert to Skeletal Mesh");
+			const TAttribute<FText> ToolTip = LOCTEXT("ConvertToSkeletalMeshToolTip", "Converts a non-Nanite static mesh to a skeletal mesh with an associated skeleton.");
+			const FSlateIcon Icon = FSlateIconFinder::FindIconForClass(USkeletalMesh::StaticClass());
+			const FToolMenuExecuteAction UIAction = FToolMenuExecuteAction::CreateLambda([](const FToolMenuContext& InMenuContext)
+			{
+				if (const UContentBrowserAssetContextMenuContext* Context = InMenuContext.FindContext<UContentBrowserAssetContextMenuContext>())
+				{
+					ConvertStaticMeshToSkeletalMeshInteractive(
+						Context->GetSelectedAssetsOfType(UStaticMesh::StaticClass()));
+				}
+			});
+
+			InSection.AddMenuEntry("ConvertToSkeletalMesh", Label, ToolTip, Icon, UIAction);
+		}));
+	}
 }
 
 
diff --git a/Engine/Plugins/Experimental/Animation/SkeletalMeshModelingTools/Source/SkeletalMeshModelingTools/SkeletalMeshModelingTools.Build.cs b/Engine/Plugins/Experimental/Animation/SkeletalMeshModelingTools/Source/SkeletalMeshModelingTools/SkeletalMeshModelingTools.Build.cs
--- a/Engine/Plugins/Experimental/Animation/SkeletalMeshModelingTools/Source/SkeletalMeshModelingTools/SkeletalMeshModelingTools.Build.cs
+++ b/Engine/Plugins/Experimental/Animation/SkeletalMeshModelingTools/Source/SkeletalMeshModelingTools/SkeletalMeshModelingTools.Build.cs
@@ -15,25 +15,27 @@ public class SkeletalMeshModelingTools : ModuleRules
 		PrivateDependencyModuleNames.AddRange(
 			new string[]
 			{
+				"ContentBrowser",
 				"Core",
 				"CoreUObject",
-				"Slate",
-				"SlateCore",
 				"EditorFramework",
+				"EditorInteractiveToolsFramework",
 				"Engine",
-				"ToolMenus",
-				"UnrealEd",
 				"InputCore",
-				"ModelingComponentsEditorOnly",
+				"InteractiveToolsFramework",
 				"MeshModelingTools",
-				"MeshModelingToolsExp",
 				"MeshModelingToolsEditorOnly",
 				"MeshModelingToolsEditorOnlyExp",
+				"MeshModelingToolsExp",
+				"ModelingComponentsEditorOnly",
 				"ModelingToolsEditorMode",
-				"InteractiveToolsFramework",
-				"EditorInteractiveToolsFramework",
+				"SkeletalMeshUtilitiesCommon",
+				"Slate",
+				"SlateCore",
 				"StylusInput",
+				"ToolMenus",
 				"ToolWidgets",
+				"UnrealEd", 
 			}
 		);
 	}
diff --git a/Engine/Source/Developer/SkeletalMeshUtilitiesCommon/Private/StaticToSkeletalMeshConverter.cpp b/Engine/Source/Developer/SkeletalMeshUtilitiesCommon/Private/StaticToSkeletalMeshConverter.cpp
new file mode 100644
--- /dev/null
+++ b/Engine/Source/Developer/SkeletalMeshUtilitiesCommon/Private/StaticToSkeletalMeshConverter.cpp
@@ -0,0 +1,328 @@
+﻿// Copyright Epic Games, Inc. All Rights Reserved.
+
+#include "StaticToSkeletalMeshConverter.h"
+
+#if WITH_EDITOR
+
+#include "Animation/Skeleton.h"
+#include "Engine/SkeletalMesh.h"
+#include "Engine/SkinnedAssetCommon.h"
+#include "Engine/StaticMesh.h"
+#include "Interfaces/ITargetPlatformManagerModule.h"
+#include "LODUtilities.h"
+#include "MeshDescription.h"
+#include "MeshUtilities.h"
+#include "Modules/ModuleManager.h"
+#include "ReferenceSkeleton.h"
+#include "Rendering/SkeletalMeshLODImporterData.h"
+#include "Rendering/SkeletalMeshModel.h"
+#include "SkeletalMeshAttributes.h"
+
+
+DEFINE_LOG_CATEGORY_STATIC(LogStaticToSkeletalMeshConverter, Log, All);
+
+static const FName RootBoneName("Root");
+static const TCHAR* JointBaseName(TEXT("Joint"));
+
+
+USkeleton* FStaticToSkeletalMeshConverter::CreateSkeletonFromStaticMesh(
+	UObject *InOuter,
+	const FName InName,
+	const EObjectFlags InFlags,
+	const UStaticMesh* InStaticMesh,
+	const FVector& InRelativeRootPosition
+	)
+{
+	if (!ensure(InStaticMesh))
+	{
+		return nullptr;
+	}
+	
+	const FBox Bounds = InStaticMesh->GetBoundingBox();
+	const FVector RootPosition = Bounds.Min + (Bounds.Max - Bounds.Min) * InRelativeRootPosition;
+	FTransform RootTransform(FTransform::Identity);
+	RootTransform.SetTranslation(RootPosition);
+
+	USkeleton* Skeleton = NewObject<USkeleton>(InOuter, InName, InFlags);
+
+	FReferenceSkeletonModifier Modifier(Skeleton);
+	Modifier.Add(FMeshBoneInfo(RootBoneName, RootBoneName.ToString(), INDEX_NONE), RootTransform);
+
+	return Skeleton;
+}
+
+
+USkeleton* FStaticToSkeletalMeshConverter::CreateSkeletonFromStaticMesh(
+	UObject *InOuter,
+	const FName InName,
+	const EObjectFlags InFlags,
+	const UStaticMesh* InStaticMesh,
+	const FVector& InRelativeRootPosition,
+	const FVector& InRelativeEndEffectorPosition,
+	const int32 InIntermediaryJointCount
+	)
+{
+	if (FMath::IsNearlyZero(FVector::DistSquared(InRelativeEndEffectorPosition, InRelativeRootPosition)))
+	{
+		return CreateSkeletonFromStaticMesh(InOuter, InName, InFlags, InStaticMesh, InRelativeRootPosition);
+	}
+
+	if (!ensure(InStaticMesh))
+	{
+		return nullptr;
+	}
+	
+	const FBox Bounds = InStaticMesh->GetBoundingBox();
+	const FVector RootPosition = Bounds.Min + (Bounds.Max - Bounds.Min) * InRelativeRootPosition;
+	const FVector EndEffectorPosition = Bounds.Min + (Bounds.Max - Bounds.Min) * InRelativeEndEffectorPosition;
+
+	// Find a rough rotation we can use
+	const FQuat Rotation = FQuat::FindBetweenVectors(FVector::ZAxisVector, EndEffectorPosition - RootPosition).GetNormalized(); 
+	
+	FTransform ParentTransform(FTransform::Identity);
+	ParentTransform.SetTranslation(RootPosition);
+	ParentTransform.SetRotation(Rotation);
+
+	USkeleton* Skeleton = NewObject<USkeleton>(InOuter, InName, InFlags);
+	FReferenceSkeletonModifier Modifier(Skeleton);
+	Modifier.Add(FMeshBoneInfo(RootBoneName, RootBoneName.ToString(), INDEX_NONE), ParentTransform);
+
+	for (int32 JointIndex = 0; JointIndex <= InIntermediaryJointCount; JointIndex++)
+	{
+		const double T = (JointIndex + 1.0) / (InIntermediaryJointCount + 2.0);
+		FTransform PointTransform(ParentTransform);
+		PointTransform.SetTranslation(RootPosition + (EndEffectorPosition - RootPosition) * T);
+
+		FString JointName = FString::Printf(TEXT("%s_%d"), JointBaseName, JointIndex + 1);
+		Modifier.Add(FMeshBoneInfo(FName(JointName), JointName, JointIndex), PointTransform * ParentTransform.Inverse());
+		ParentTransform = PointTransform;
+	}
+
+	return Skeleton;
+}
+
+static void CopyBuildSettings(
+	const FMeshBuildSettings& InStaticMeshBuildSettings,
+	FSkeletalMeshBuildSettings& OutSkeletalMeshBuildSettings
+	)
+{
+	OutSkeletalMeshBuildSettings.bRecomputeNormals = InStaticMeshBuildSettings.bRecomputeNormals;
+	OutSkeletalMeshBuildSettings.bRecomputeTangents = InStaticMeshBuildSettings.bRecomputeTangents;
+	OutSkeletalMeshBuildSettings.bUseMikkTSpace = InStaticMeshBuildSettings.bUseMikkTSpace;
+	OutSkeletalMeshBuildSettings.bComputeWeightedNormals = InStaticMeshBuildSettings.bComputeWeightedNormals;
+	OutSkeletalMeshBuildSettings.bRemoveDegenerates = InStaticMeshBuildSettings.bRemoveDegenerates;
+	OutSkeletalMeshBuildSettings.bUseHighPrecisionTangentBasis = InStaticMeshBuildSettings.bUseHighPrecisionTangentBasis;
+	OutSkeletalMeshBuildSettings.bUseFullPrecisionUVs = InStaticMeshBuildSettings.bUseFullPrecisionUVs;
+	OutSkeletalMeshBuildSettings.bUseBackwardsCompatibleF16TruncUVs = InStaticMeshBuildSettings.bUseBackwardsCompatibleF16TruncUVs;
+	// The rest we leave at defaults.
+}
+
+static SkeletalMeshOptimizationImportance ConvertOptimizationImportance(
+	 EMeshFeatureImportance::Type InStaticMeshImportance)
+{
+	switch(InStaticMeshImportance)
+	{
+	default:
+	case EMeshFeatureImportance::Off:		return SMOI_Highest;
+	case EMeshFeatureImportance::Lowest:	return SMOI_Lowest;
+	case EMeshFeatureImportance::Low:		return SMOI_Low;
+	case EMeshFeatureImportance::Normal:	return SMOI_Normal;
+	case EMeshFeatureImportance::High:		return SMOI_High;
+	case EMeshFeatureImportance::Highest:	return SMOI_Highest;
+	}
+}
+
+static void CopyReductionSettings(
+	const FMeshReductionSettings& InStaticMeshReductionSettings,
+	FSkeletalMeshOptimizationSettings& OutSkeletalMeshReductionSettings
+	)
+{
+	// Copy the reduction settings as closely as we can. 
+	OutSkeletalMeshReductionSettings.NumOfTrianglesPercentage = InStaticMeshReductionSettings.PercentTriangles;
+	OutSkeletalMeshReductionSettings.NumOfVertPercentage = InStaticMeshReductionSettings.PercentVertices;
+	
+	OutSkeletalMeshReductionSettings.WeldingThreshold = InStaticMeshReductionSettings.WeldingThreshold;
+	OutSkeletalMeshReductionSettings.NormalsThreshold = InStaticMeshReductionSettings.HardAngleThreshold;
+	OutSkeletalMeshReductionSettings.bRecalcNormals = InStaticMeshReductionSettings.bRecalculateNormals;
+	
+	OutSkeletalMeshReductionSettings.BaseLOD = InStaticMeshReductionSettings.BaseLODModel;
+	
+	OutSkeletalMeshReductionSettings.SilhouetteImportance = ConvertOptimizationImportance(InStaticMeshReductionSettings.SilhouetteImportance);
+	OutSkeletalMeshReductionSettings.TextureImportance = ConvertOptimizationImportance(InStaticMeshReductionSettings.TextureImportance);
+	OutSkeletalMeshReductionSettings.ShadingImportance = ConvertOptimizationImportance(InStaticMeshReductionSettings.ShadingImportance);
+	
+	switch(InStaticMeshReductionSettings.TerminationCriterion)
+	{
+	case EStaticMeshReductionTerimationCriterion::Triangles:
+		OutSkeletalMeshReductionSettings.TerminationCriterion = SMTC_NumOfTriangles;
+		break;
+	case EStaticMeshReductionTerimationCriterion::Vertices:
+		OutSkeletalMeshReductionSettings.TerminationCriterion = SMTC_NumOfVerts;
+		break;
+	case EStaticMeshReductionTerimationCriterion::Any:
+		OutSkeletalMeshReductionSettings.TerminationCriterion = SMTC_TriangleOrVert;
+		break;
+	}
+}
+
+
+
+static bool AddLODFromMeshDescription(
+	const FStaticMeshSourceModel& InStaticMeshSourceModel,
+	USkeletalMesh* InSkeletalMesh,
+	const int32 InLODIndex,
+	const FBoneIndexType InBoneIndex,
+	IMeshUtilities& InMeshUtilities
+	)
+{
+	// Always copy the build and reduction settings. 
+	FSkeletalMeshLODInfo& SkeletalLODInfo = InSkeletalMesh->AddLODInfo();
+
+	SkeletalLODInfo.ScreenSize = InStaticMeshSourceModel.ScreenSize;
+	CopyBuildSettings(InStaticMeshSourceModel.BuildSettings, SkeletalLODInfo.BuildSettings);
+	CopyReductionSettings(InStaticMeshSourceModel.ReductionSettings, SkeletalLODInfo.ReductionSettings);
+
+	FSkeletalMeshModel* ImportedModels = InSkeletalMesh->GetImportedModel();
+	ImportedModels->LODModels.Add(new FSkeletalMeshLODModel);
+	FSkeletalMeshLODModel& SkeletalMeshModel = ImportedModels->LODModels.Last(); 
+	
+	if (InStaticMeshSourceModel.IsMeshDescriptionValid())
+	{
+		FMeshDescription SkeletalMeshGeometry;
+		if (!InStaticMeshSourceModel.CloneMeshDescription(SkeletalMeshGeometry))
+		{
+			return false;
+		}
+		
+		FSkeletalMeshAttributes SkeletalMeshAttributes(SkeletalMeshGeometry);
+		SkeletalMeshAttributes.Register();
+
+		// Full binding to the root bone.
+		FSkinWeightsVertexAttributesRef SkinWeights = SkeletalMeshAttributes.GetVertexSkinWeights();
+		UE::AnimationCore::FBoneWeight RootInfluence(InBoneIndex, 1.0f);
+		UE::AnimationCore::FBoneWeights RootBinding = UE::AnimationCore::FBoneWeights::Create({RootInfluence});
+		
+		for (const FVertexID VertexID: SkeletalMeshGeometry.Vertices().GetElementIDs())
+		{
+			SkinWeights.Set(VertexID, RootBinding);
+		}
+
+		FSkeletalMeshImportData SkeletalMeshImportGeometry = FSkeletalMeshImportData::CreateFromMeshDescription(SkeletalMeshGeometry);
+
+		// We need at least one set of texture coordinates. Always.
+		SkeletalMeshModel.NumTexCoords = FMath::Max<uint32>(1, SkeletalMeshImportGeometry.NumTexCoords);
+		
+		// Data needed by BuildSkeletalMesh
+		TArray<FVector3f> LODPoints;
+		TArray<SkeletalMeshImportData::FMeshWedge> LODWedges;
+		TArray<SkeletalMeshImportData::FMeshFace> LODFaces;
+		TArray<SkeletalMeshImportData::FVertInfluence> LODInfluences;
+		TArray<int32> LODPointToRawMap;
+		SkeletalMeshImportGeometry.CopyLODImportData( LODPoints, LODWedges, LODFaces, LODInfluences, LODPointToRawMap );
+
+		IMeshUtilities::MeshBuildOptions BuildOptions;
+		BuildOptions.TargetPlatform = GetTargetPlatformManagerRef().GetRunningTargetPlatform();
+		BuildOptions.FillOptions(SkeletalLODInfo.BuildSettings);
+
+		TArray<FText> WarningMessages;
+		if (!InMeshUtilities.BuildSkeletalMesh(SkeletalMeshModel, InSkeletalMesh->GetPathName(), InSkeletalMesh->GetRefSkeleton(), LODInfluences, LODWedges, LODFaces, LODPoints, LODPointToRawMap, BuildOptions, &WarningMessages, nullptr))
+		{
+			for(const FText& Message: WarningMessages)
+			{
+				UE_LOG(LogStaticToSkeletalMeshConverter, Warning, TEXT("%s"), *Message.ToString());
+			}
+			return false;
+		}
+		
+		InSkeletalMesh->SaveLODImportedData(InLODIndex, SkeletalMeshImportGeometry);
+	}
+	else
+	{
+		FSkeletalMeshUpdateContext UpdateContext;
+		UpdateContext.SkeletalMesh = InSkeletalMesh;
+		
+		FLODUtilities::SimplifySkeletalMeshLOD(UpdateContext, InLODIndex, GetTargetPlatformManagerRef().GetRunningTargetPlatform());
+	}
+	
+	return true;
+}
+
+
+
+USkeletalMesh* FStaticToSkeletalMeshConverter::CreateSkeletalMeshFromStaticMesh(
+	UObject *InOuter,
+	const FName InName,
+	const EObjectFlags InFlags,
+	const UStaticMesh* InStaticMesh,
+	const FReferenceSkeleton& InReferenceSkeleton,
+	const FName InBindBone 
+)
+{
+	if (!ensure(InStaticMesh))
+	{
+		return nullptr;
+	}
+
+	int32 BoneIndex = 0;
+	if (!InBindBone.IsNone())
+	{
+		BoneIndex = InReferenceSkeleton.FindRawBoneIndex(InBindBone);
+		if (!ensure(BoneIndex != INDEX_NONE))
+		{
+			return nullptr;
+		}
+	}
+
+	USkeletalMesh* SkeletalMesh = NewObject<USkeletalMesh>(InOuter, InName, InFlags);
+
+	// This ensures that the render data gets built before we return, by calling PostEditChange when we fall out of scope.
+	FScopedSkeletalMeshPostEditChange ScopedPostEditChange( SkeletalMesh );
+	SkeletalMesh->PreEditChange( nullptr );
+	SkeletalMesh->SetRefSkeleton(InReferenceSkeleton);
+	
+	// Calculate the initial pose from the reference skeleton.
+	SkeletalMesh->CalculateInvRefMatrices();
+
+	IMeshUtilities& MeshUtilities = FModuleManager::Get().LoadModuleChecked<IMeshUtilities>( "MeshUtilities" );
+	
+	// Copy the LODs and LOD settings over (as close as we can).
+	for (int32 Index = 0; Index < InStaticMesh->GetNumSourceModels(); Index++)
+	{
+		const FStaticMeshSourceModel& StaticMeshSourceModel = InStaticMesh->GetSourceModel(Index);
+		
+		if (!AddLODFromMeshDescription(
+			StaticMeshSourceModel, SkeletalMesh, Index, static_cast<FBoneIndexType>(BoneIndex), MeshUtilities))
+		{
+			// If we didn't get a model for LOD index 0, we don't have a mesh. Bail out.
+			if (Index == 0)
+			{
+				return nullptr;
+			}
+
+			// Otherwise, we have a model, so let's continue with what we have.
+			break;
+		}
+	}
+	
+	// Convert the materials over.
+	TArray<FSkeletalMaterial> Materials;
+	for (const FStaticMaterial& StaticMaterial: InStaticMesh->GetStaticMaterials())
+	{
+		FSkeletalMaterial Material(
+			StaticMaterial.MaterialInterface,
+			StaticMaterial.MaterialSlotName);
+		
+		Materials.Add(Material);
+	}
+
+	SkeletalMesh->SetMaterials(Materials);
+	
+	// Set the bounds from the static mesh, including the extensions, otherwise it won't render properly (among other things).
+	SkeletalMesh->SetImportedBounds( InStaticMesh->GetBounds() );
+	SkeletalMesh->SetPositiveBoundsExtension(InStaticMesh->GetPositiveBoundsExtension());
+	SkeletalMesh->SetNegativeBoundsExtension(InStaticMesh->GetNegativeBoundsExtension());
+
+	return SkeletalMesh;
+}
+
+#endif
diff --git a/Engine/Source/Developer/SkeletalMeshUtilitiesCommon/Public/StaticToSkeletalMeshConverter.h b/Engine/Source/Developer/SkeletalMeshUtilitiesCommon/Public/StaticToSkeletalMeshConverter.h
new file mode 100644
--- /dev/null
+++ b/Engine/Source/Developer/SkeletalMeshUtilitiesCommon/Public/StaticToSkeletalMeshConverter.h
@@ -0,0 +1,78 @@
+﻿// Copyright Epic Games, Inc. All Rights Reserved.
+
+#pragma once
+
+#if WITH_EDITOR
+
+#include "Math/MathFwd.h"
+#include "UObject/NameTypes.h"
+#include "UObject/ObjectMacros.h"
+
+
+class UObject;
+class USkeleton;
+class UStaticMesh;
+class USkeletalMesh;
+struct FReferenceSkeleton;
+
+
+struct SKELETALMESHUTILITIESCOMMON_API FStaticToSkeletalMeshConverter
+{
+	/** Creates a skeleton from a static mesh with a single root bone.
+	 *
+	 *  @param InStaticMesh The static mesh whose bounding box will be used as a reference.
+	 *  @param InRelativeRootPosition The relative root position in a unit cube that gets scaled up to match the
+	 *    bbox of the static mesh. For example, given FVector(0.5, 0.5, 0.0), the root bone location will be placed
+	 *    at the center of the bottom of the static mesh's bbox.
+	 *  NOTE: The owner upon return will be the transient package, re-outer to the correct package as needed. 
+	 */
+	static USkeleton* CreateSkeletonFromStaticMesh(
+		UObject *InOuter,
+		const FName InName,
+		const EObjectFlags InFlags,
+		const UStaticMesh* InStaticMesh,
+		const FVector& InRelativeRootPosition 
+		);
+
+	/** Creates a skeleton from a static mesh with a bone chain going from root to end effector, where the intermediary
+	 *  bones are distributed evenly along a line between the two.
+	 *
+	 *  @param InStaticMesh The static mesh whose bounding box will be used as a reference.
+	 *  @param InRelativeRootPosition The relative root position in a unit cube that gets scaled up to match the
+	 *    bbox of the static mesh. For example, given FVector(0.5, 0.5, 0.0), the root bone location will be placed
+	 *    at the center of the bottom of the static mesh's bbox.
+	 *  @param InRelativeEndEffectorPosition The end effector position, positioned in the same manner as the root
+	 *    position. If the end effector is in the same location as the root, only the root bone is created.
+	 *  @param InIntermediaryJointCount Number of joints to create between the root and the end effector.  
+	 *  NOTE: The owner upon return will be the transient package, re-outer to the correct package as needed. 
+	 */
+	static USkeleton* CreateSkeletonFromStaticMesh(
+		UObject *InOuter,
+		const FName InName,
+		const EObjectFlags InFlags,
+		const UStaticMesh* InStaticMesh,
+		const FVector& InRelativeRootPosition,
+		const FVector& InRelativeEndEffectorPosition,
+		const int32 InIntermediaryJointCount
+		);
+
+	/** Create a skeletal mesh from the given static mesh and a skeleton. The mesh will initially be created with a
+	 *  rigid binding on the root bone. Use SetRigidBinding or SetSmoothBinding to override this given binding.
+	 *  @param InOuter The skeletal mesh's outer/parent object.
+	 *  @param InName The name to give the newly created skeletal mesh.
+	 *  @param InFlags The object creation flags to use.
+	 *  @param InStaticMesh The static mesh to convert from.
+	 *  @param InReferenceSkeleton The reference skeleton to use.
+	 *  @param InBindBone The bone to bind to. If no bone name is given, the binding defaults to the root bone.
+	 */
+	static USkeletalMesh* CreateSkeletalMeshFromStaticMesh(
+		UObject *InOuter,
+		const FName InName,
+		const EObjectFlags InFlags,
+		const UStaticMesh* InStaticMesh,
+		const FReferenceSkeleton& InReferenceSkeleton,
+		const FName InBindBone = NAME_None 
+		);
+};
+
+#endif
diff --git a/Engine/Source/Developer/SkeletalMeshUtilitiesCommon/SkeletalMeshUtilitiesCommon.Build.cs b/Engine/Source/Developer/SkeletalMeshUtilitiesCommon/SkeletalMeshUtilitiesCommon.Build.cs
--- a/Engine/Source/Developer/SkeletalMeshUtilitiesCommon/SkeletalMeshUtilitiesCommon.Build.cs
+++ b/Engine/Source/Developer/SkeletalMeshUtilitiesCommon/SkeletalMeshUtilitiesCommon.Build.cs
@@ -13,13 +13,15 @@ namespace UnrealBuildTool.Rules
             PrivateDependencyModuleNames.AddRange(
                 new string[]
                 {
-                    "Core",
+					"ClothingSystemRuntimeCommon",
+					"Core",
 					"CoreUObject",
-					"RenderCore",
-					"MeshUtilitiesCommon",
-					"ImageCore",
 					"Engine",
-					"ClothingSystemRuntimeCommon",
+					"ImageCore",
+					"MeshDescription",
+					"MeshUtilitiesCommon",
+					"RenderCore",
+					"SkeletalMeshDescription",
 					"Slate",
                 }
             );
diff --git a/Engine/Source/Runtime/Engine/Classes/Engine/SkinnedAssetCommon.h b/Engine/Source/Runtime/Engine/Classes/Engine/SkinnedAssetCommon.h
--- a/Engine/Source/Runtime/Engine/Classes/Engine/SkinnedAssetCommon.h
+++ b/Engine/Source/Runtime/Engine/Classes/Engine/SkinnedAssetCommon.h
@@ -242,6 +242,21 @@ struct FSkeletalMaterial
 
 	}
 
+	FSkeletalMaterial(
+		UMaterialInterface* InMaterialInterface,
+		FName InMaterialSlotName,
+		FName InImportedMaterialSlotName = NAME_None)
+		: MaterialInterface( InMaterialInterface )
+		, MaterialSlotName(InMaterialSlotName)
+#if WITH_EDITORONLY_DATA
+		, bEnableShadowCasting_DEPRECATED(true)
+		, bRecomputeTangent_DEPRECATED(false)
+		, ImportedMaterialSlotName(InImportedMaterialSlotName)
+#endif //WITH_EDITORONLY_DATA
+	{
+
+	}
+
 	FSkeletalMaterial( class UMaterialInterface* InMaterialInterface
 						, bool bInEnableShadowCasting = true
 						, bool bInRecomputeTangent = false
diff --git a/Engine/Source/Runtime/Engine/Private/SkeletalMesh.cpp b/Engine/Source/Runtime/Engine/Private/SkeletalMesh.cpp
--- a/Engine/Source/Runtime/Engine/Private/SkeletalMesh.cpp
+++ b/Engine/Source/Runtime/Engine/Private/SkeletalMesh.cpp
@@ -535,7 +535,7 @@ void USkeletalMesh::CalculateExtendedBounds()
 	// Convert back to Origin, Extent and update SphereRadius
 	CalculatedBounds.Origin = (Min + Max) / 2;
 	CalculatedBounds.BoxExtent = (Max - Min) / 2;
-	CalculatedBounds.SphereRadius = CalculatedBounds.BoxExtent.GetAbsMax();
+	CalculatedBounds.SphereRadius = CalculatedBounds.BoxExtent.Size();
 
 	SetExtendedBounds(CalculatedBounds);
 }
