Fix build break on Mac

#rnx

diff --git a/Engine/Source/Developer/SkeletalMeshUtilitiesCommon/SkeletalMeshUtilitiesCommon.Build.cs b/Engine/Source/Developer/SkeletalMeshUtilitiesCommon/SkeletalMeshUtilitiesCommon.Build.cs
--- a/Engine/Source/Developer/SkeletalMeshUtilitiesCommon/SkeletalMeshUtilitiesCommon.Build.cs
+++ b/Engine/Source/Developer/SkeletalMeshUtilitiesCommon/SkeletalMeshUtilitiesCommon.Build.cs
@@ -23,6 +23,7 @@ namespace UnrealBuildTool.Rules
 					"RenderCore",
 					"SkeletalMeshDescription",
 					"Slate",
+					"StaticMeshDescription",
                 }
             );
         }
