Allow users to limit the number of bone influences per vertex below what the project supports.

BoneInfluencesLimit can be set on each LOD of an asset. DefaultBoneInfluencesLimit is a per-platform global default that can be set from the project settings.

This allows you to, for example, enable Unlimited Bone Influences on a project while ensuring that meshes don't use more than 8 influences by default, even if they've been imported with more.

diff --git a/Engine/Source/Developer/MeshBuilder/Private/SkeletalMeshBuilder.cpp b/Engine/Source/Developer/MeshBuilder/Private/SkeletalMeshBuilder.cpp
--- a/Engine/Source/Developer/MeshBuilder/Private/SkeletalMeshBuilder.cpp
+++ b/Engine/Source/Developer/MeshBuilder/Private/SkeletalMeshBuilder.cpp
@@ -160,7 +160,7 @@ bool FSkeletalMeshBuilder::Build(const FSkeletalMeshBuildParameters& SkeletalMes
 		for (int32 SkinProfileIndex = 0; SkinProfileIndex < SkinProfiles.Num(); ++SkinProfileIndex)
 		{
 			const FSkinWeightProfileInfo& ProfileInfo = SkinProfiles[SkinProfileIndex];
-			FLODUtilities::UpdateAlternateSkinWeights(SkeletalMesh, ProfileInfo.Name, LODIndex, Options.OverlappingThresholds, !Options.bComputeNormals, !Options.bComputeTangents, Options.bUseMikkTSpace, Options.bComputeWeightedNormals);
+			FLODUtilities::UpdateAlternateSkinWeights(SkeletalMesh, ProfileInfo.Name, LODIndex, Options);
 		}
 
 		
diff --git a/Engine/Source/Developer/MeshUtilities/Private/MeshUtilities.cpp b/Engine/Source/Developer/MeshUtilities/Private/MeshUtilities.cpp
--- a/Engine/Source/Developer/MeshUtilities/Private/MeshUtilities.cpp
+++ b/Engine/Source/Developer/MeshUtilities/Private/MeshUtilities.cpp
@@ -4055,6 +4055,15 @@ public:
 					}
 
 					InfluenceCount = FMath::Min<uint32>(InfluenceCount, MAX_TOTAL_INFLUENCES);
+
+					// Apply bone influence limit from asset
+					{
+						const int32 InfluenceLimit = FGPUBaseSkinVertexFactory::GetBoneInfluenceLimitForAsset(
+							BuildData.BuildOptions.BoneInfluenceLimit, BuildData.BuildOptions.TargetPlatform);
+
+						InfluenceCount = FMath::Min<uint32>(InfluenceCount, InfluenceLimit);
+					}
+
 					if (InfluenceCount > EXTRA_BONE_INFLUENCES && !FGPUBaseSkinVertexFactory::UseUnlimitedBoneInfluences(InfluenceCount))
 					{
 						// Add a single warning message, there could be too many
diff --git a/Engine/Source/Developer/MeshUtilities/Public/MeshUtilities.h b/Engine/Source/Developer/MeshUtilities/Public/MeshUtilities.h
--- a/Engine/Source/Developer/MeshUtilities/Public/MeshUtilities.h
+++ b/Engine/Source/Developer/MeshUtilities/Public/MeshUtilities.h
@@ -152,7 +152,8 @@ public:
 	struct MeshBuildOptions
 	{
 		MeshBuildOptions()
-		: bRemoveDegenerateTriangles(false)
+		: BoneInfluenceLimit(0)
+		, bRemoveDegenerateTriangles(false)
 		, bComputeNormals(true)
 		, bComputeTangents(true)
 		, bUseMikkTSpace(false)
@@ -161,6 +162,7 @@ public:
 		{
 		}
 
+		int32 BoneInfluenceLimit;
 		bool bRemoveDegenerateTriangles;
 		bool bComputeNormals;
 		bool bComputeTangents;
@@ -175,6 +177,7 @@ public:
 			OverlappingThresholds.ThresholdTangentNormal = SkeletalMeshBuildSettings.ThresholdTangentNormal;
 			OverlappingThresholds.ThresholdUV = SkeletalMeshBuildSettings.ThresholdUV;
 			OverlappingThresholds.MorphThresholdPosition = SkeletalMeshBuildSettings.MorphThresholdPosition;
+			BoneInfluenceLimit = SkeletalMeshBuildSettings.BoneInfluenceLimit;
 			bComputeNormals = SkeletalMeshBuildSettings.bRecomputeNormals;
 			bComputeTangents = SkeletalMeshBuildSettings.bRecomputeTangents;
 			bUseMikkTSpace = SkeletalMeshBuildSettings.bUseMikkTSpace;
diff --git a/Engine/Source/Developer/SkeletalMeshUtilitiesCommon/Private/LODUtilities.cpp b/Engine/Source/Developer/SkeletalMeshUtilitiesCommon/Private/LODUtilities.cpp
--- a/Engine/Source/Developer/SkeletalMeshUtilitiesCommon/Private/LODUtilities.cpp
+++ b/Engine/Source/Developer/SkeletalMeshUtilitiesCommon/Private/LODUtilities.cpp
@@ -2155,7 +2155,7 @@ void MatchVertexIndexUsingPosition(
 	}
 }
 
-bool FLODUtilities::UpdateAlternateSkinWeights(USkeletalMesh* SkeletalMeshDest, const FName& ProfileNameDest, int32 LODIndexDest, FOverlappingThresholds OverlappingThresholds, bool ShouldImportNormals, bool ShouldImportTangents, bool bUseMikkTSpace, bool bComputeWeightedNormals)
+bool FLODUtilities::UpdateAlternateSkinWeights(USkeletalMesh* SkeletalMeshDest, const FName& ProfileNameDest, int32 LODIndexDest, const IMeshUtilities::MeshBuildOptions& Options)
 {
 	//Grab all the destination structure
 	check(SkeletalMeshDest);
@@ -2168,12 +2168,48 @@ bool FLODUtilities::UpdateAlternateSkinWeights(USkeletalMesh* SkeletalMeshDest,
 		//Very old asset will not have this data, we cannot add alternate until the asset is reimported
 		return false;
 	}
+	
 	FSkeletalMeshImportData ImportDataDest;
 	SkeletalMeshDest->LoadLODImportedData(LODIndexDest, ImportDataDest);
-	return UpdateAlternateSkinWeights(LODModelDest, ImportDataDest, SkeletalMeshDest, SkeletalMeshDest->GetRefSkeleton(), ProfileNameDest, LODIndexDest, OverlappingThresholds, ShouldImportNormals, ShouldImportTangents, bUseMikkTSpace, bComputeWeightedNormals);
+	
+	return UpdateAlternateSkinWeights(
+		LODModelDest,
+		ImportDataDest,
+		SkeletalMeshDest,
+		SkeletalMeshDest->GetRefSkeleton(),
+		ProfileNameDest,
+		LODIndexDest,
+		Options);
 }
 
-bool FLODUtilities::UpdateAlternateSkinWeights(FSkeletalMeshLODModel& LODModelDest, FSkeletalMeshImportData& ImportDataDest, USkeletalMesh* SkeletalMeshDest, const FReferenceSkeleton& RefSkeleton, const FName& ProfileNameDest, int32 LODIndexDest, FOverlappingThresholds OverlappingThresholds, bool ShouldImportNormals, bool ShouldImportTangents, bool bUseMikkTSpace, bool bComputeWeightedNormals)
+bool FLODUtilities::UpdateAlternateSkinWeights(
+	USkeletalMesh* SkeletalMeshDest,
+	const FName& ProfileNameDest,
+	int32 LODIndexDest,
+	FOverlappingThresholds OverlappingThresholds,
+	bool ShouldImportNormals,
+	bool ShouldImportTangents,
+	bool bUseMikkTSpace,
+	bool bComputeWeightedNormals)
+{
+	IMeshUtilities::MeshBuildOptions Options;
+	Options.OverlappingThresholds = OverlappingThresholds;
+	Options.bComputeNormals = !ShouldImportNormals;
+	Options.bComputeTangents = !ShouldImportTangents;
+	Options.bUseMikkTSpace = bUseMikkTSpace;
+	Options.bComputeWeightedNormals = bComputeWeightedNormals;
+
+	return UpdateAlternateSkinWeights(SkeletalMeshDest, ProfileNameDest, LODIndexDest, Options);
+}
+
+bool FLODUtilities::UpdateAlternateSkinWeights(
+	FSkeletalMeshLODModel& LODModelDest,
+	FSkeletalMeshImportData& ImportDataDest,
+	USkeletalMesh* SkeletalMeshDest,
+	const FReferenceSkeleton& RefSkeleton,
+	const FName& ProfileNameDest,
+	int32 LODIndexDest,
+	const IMeshUtilities::MeshBuildOptions& Options)
 {
 	//Ensure log message only once
 	bool bNoMatchMsgDone = false;
@@ -2527,14 +2563,10 @@ bool FLODUtilities::UpdateAlternateSkinWeights(FSkeletalMeshLODModel& LODModelDe
 	ImportDataDest.CopyLODImportData(LODPointsDest, LODWedgesDest, LODFacesDest, LODInfluencesDest, LODPointToRawMapDest);
 
 	//Set the options with the current asset build options
-	IMeshUtilities::MeshBuildOptions BuildOptions;
-	BuildOptions.OverlappingThresholds = OverlappingThresholds;
-	BuildOptions.bComputeNormals = !ShouldImportNormals || !ImportDataDest.bHasNormals;
-	BuildOptions.bComputeTangents = !ShouldImportTangents || !ImportDataDest.bHasTangents;
-	BuildOptions.bUseMikkTSpace = (bUseMikkTSpace) && (!ShouldImportNormals || !ShouldImportTangents);
-	BuildOptions.bComputeWeightedNormals = bComputeWeightedNormals;
-	BuildOptions.bRemoveDegenerateTriangles = false;
-	BuildOptions.TargetPlatform = GetTargetPlatformManagerRef().GetRunningTargetPlatform();
+	IMeshUtilities::MeshBuildOptions BuildOptions = Options;
+	BuildOptions.bComputeNormals = Options.bComputeNormals || !ImportDataDest.bHasNormals;
+	BuildOptions.bComputeTangents = Options.bComputeTangents || !ImportDataDest.bHasTangents;
+	BuildOptions.bUseMikkTSpace = (Options.bUseMikkTSpace) && (Options.bComputeNormals || Options.bComputeTangents);
 
 	//Build the skeletal mesh asset
 	IMeshUtilities& MeshUtilities = FModuleManager::Get().LoadModuleChecked<IMeshUtilities>("MeshUtilities");
@@ -2549,12 +2581,41 @@ bool FLODUtilities::UpdateAlternateSkinWeights(FSkeletalMeshLODModel& LODModelDe
 	//Re-Apply the user section changes, the UserSectionsData is map to original section and should match the builded LODModel
 	LODModelDest.SyncronizeUserSectionsDataArray();
 
-	RegenerateAllImportSkinWeightProfileData(LODModelDest);
+	RegenerateAllImportSkinWeightProfileData(LODModelDest, Options.BoneInfluenceLimit, Options.TargetPlatform);
 	
 	return bBuildSuccess;
 }
 
-bool FLODUtilities::UpdateAlternateSkinWeights(USkeletalMesh* SkeletalMeshDest, const FName& ProfileNameDest, USkeletalMesh* SkeletalMeshSrc, int32 LODIndexDest, int32 LODIndexSrc, FOverlappingThresholds OverlappingThresholds, bool ShouldImportNormals, bool ShouldImportTangents, bool bUseMikkTSpace, bool bComputeWeightedNormals)
+bool FLODUtilities::UpdateAlternateSkinWeights(
+	FSkeletalMeshLODModel& LODModelDest,
+	FSkeletalMeshImportData& ImportDataDest,
+	USkeletalMesh* SkeletalMeshDest,
+	const FReferenceSkeleton& RefSkeleton,
+	const FName& ProfileNameDest,
+	int32 LODIndexDest,
+	FOverlappingThresholds OverlappingThresholds,
+	bool ShouldImportNormals,
+	bool ShouldImportTangents,
+	bool bUseMikkTSpace,
+	bool bComputeWeightedNormals)
+{
+	IMeshUtilities::MeshBuildOptions Options;
+	Options.OverlappingThresholds = OverlappingThresholds;
+	Options.bComputeNormals = !ShouldImportNormals;
+	Options.bComputeTangents = !ShouldImportTangents;
+	Options.bUseMikkTSpace = bUseMikkTSpace;
+	Options.bComputeWeightedNormals = bComputeWeightedNormals;
+
+	return UpdateAlternateSkinWeights(LODModelDest, ImportDataDest, SkeletalMeshDest, RefSkeleton, ProfileNameDest, LODIndexDest, Options);
+}
+
+bool FLODUtilities::UpdateAlternateSkinWeights(
+	USkeletalMesh* SkeletalMeshDest,
+	const FName& ProfileNameDest,
+	USkeletalMesh* SkeletalMeshSrc,
+	int32 LODIndexDest,
+	int32 LODIndexSrc,
+	const IMeshUtilities::MeshBuildOptions& Options)
 {
 	//Grab all the destination structure
 	check(SkeletalMeshDest);
@@ -2613,12 +2674,38 @@ bool FLODUtilities::UpdateAlternateSkinWeights(USkeletalMesh* SkeletalMeshDest,
 	if(!SkeletalMeshDest->IsLODImportedDataBuildAvailable(LODIndexDest))
 	{
 		//Build the alternate buffer with all the data into the bulk, in case the build data is not existing (old asset)
-		return UpdateAlternateSkinWeights(SkeletalMeshDest, ProfileNameDest, LODIndexDest, OverlappingThresholds, ShouldImportNormals, ShouldImportTangents, bUseMikkTSpace, bComputeWeightedNormals);
+		return UpdateAlternateSkinWeights(SkeletalMeshDest, ProfileNameDest, LODIndexDest, Options);
 	}
 	return true;
 }
 
-void FLODUtilities::GenerateImportedSkinWeightProfileData(FSkeletalMeshLODModel& LODModelDest, FImportedSkinWeightProfileData &ImportedProfileData)
+bool FLODUtilities::UpdateAlternateSkinWeights(
+	USkeletalMesh* SkeletalMeshDest,
+	const FName& ProfileNameDest,
+	USkeletalMesh* SkeletalMeshSrc,
+	int32 LODIndexDest,
+	int32 LODIndexSrc,
+	FOverlappingThresholds OverlappingThresholds,
+	bool ShouldImportNormals,
+	bool ShouldImportTangents,
+	bool bUseMikkTSpace,
+	bool bComputeWeightedNormals)
+{
+	IMeshUtilities::MeshBuildOptions Options;
+	Options.OverlappingThresholds = OverlappingThresholds;
+	Options.bComputeNormals = !ShouldImportNormals;
+	Options.bComputeTangents = !ShouldImportTangents;
+	Options.bUseMikkTSpace = bUseMikkTSpace;
+	Options.bComputeWeightedNormals = bComputeWeightedNormals;
+
+	return UpdateAlternateSkinWeights(SkeletalMeshDest, ProfileNameDest, SkeletalMeshSrc, LODIndexDest, LODIndexSrc, Options);
+}
+
+void FLODUtilities::GenerateImportedSkinWeightProfileData(
+	FSkeletalMeshLODModel& LODModelDest,
+	FImportedSkinWeightProfileData& ImportedProfileData,
+	int32 BoneInfluenceLimit,
+	const ITargetPlatform* TargetPlatform)
 {
 	//Add the override buffer with the alternate influence data
 	TArray<FSoftSkinVertex> DestinationSoftVertices;
@@ -2627,8 +2714,12 @@ void FLODUtilities::GenerateImportedSkinWeightProfileData(FSkeletalMeshLODModel&
 	TArray<FRawSkinWeight>& SkinWeights = ImportedProfileData.SkinWeights;
 	SkinWeights.Empty(DestinationSoftVertices.Num());
 
+	const int32 MaxBoneInfluencesFromProjectSettings = FGPUBaseSkinVertexFactory::UseUnlimitedBoneInfluences(MAX_TOTAL_INFLUENCES) ? MAX_TOTAL_INFLUENCES : EXTRA_BONE_INFLUENCES;
+	const int32 MaxBoneInfluencesFromAsset = FGPUBaseSkinVertexFactory::GetBoneInfluenceLimitForAsset(BoneInfluenceLimit, TargetPlatform);
+
 	//Get the maximum allow bone influence, so we can cut lowest weight properly and get the same result has the sk build
-	const int32 MaxInfluenceCount = FGPUBaseSkinVertexFactory::UseUnlimitedBoneInfluences(MAX_TOTAL_INFLUENCES) ? MAX_TOTAL_INFLUENCES : EXTRA_BONE_INFLUENCES;
+	const int32 MaxInfluenceCount = FMath::Min(MaxBoneInfluencesFromProjectSettings, MaxBoneInfluencesFromAsset);
+
 	int32 MaxNumInfluences = 0;
 
 	for (int32 VertexInstanceIndex = 0; VertexInstanceIndex < DestinationSoftVertices.Num(); ++VertexInstanceIndex)
@@ -2641,7 +2732,7 @@ void FLODUtilities::GenerateImportedSkinWeightProfileData(FSkeletalMeshLODModel&
 			continue;
 		}
 		FSkelMeshSection& Section = LODModelDest.Sections[SectionIndex];
-		const TArray<FBoneIndexType> SectionBoneMap = Section.BoneMap;
+		const TArray<FBoneIndexType>& SectionBoneMap = Section.BoneMap;
 		const FSoftSkinVertex& Vertex = DestinationSoftVertices[VertexInstanceIndex];
 		const int32 VertexIndex = LODModelDest.MeshToImportVertexMap[VertexInstanceIndex];
 		check(VertexIndex >= 0 && VertexIndex <= LODModelDest.MaxImportVertex);
@@ -2712,11 +2803,11 @@ void FLODUtilities::GenerateImportedSkinWeightProfileData(FSkeletalMeshLODModel&
 	}
 }
 
-void FLODUtilities::RegenerateAllImportSkinWeightProfileData(FSkeletalMeshLODModel& LODModelDest)
+void FLODUtilities::RegenerateAllImportSkinWeightProfileData(FSkeletalMeshLODModel& LODModelDest, int32 BoneInfluenceLimit, const ITargetPlatform* TargetPlatform)
 {
 	for (TPair<FName, FImportedSkinWeightProfileData>& ProfilePair : LODModelDest.SkinWeightProfiles)
 	{
-		GenerateImportedSkinWeightProfileData(LODModelDest, ProfilePair.Value);
+		GenerateImportedSkinWeightProfileData(LODModelDest, ProfilePair.Value, BoneInfluenceLimit, TargetPlatform);
 	}
 }
 
diff --git a/Engine/Source/Developer/SkeletalMeshUtilitiesCommon/Public/LODUtilities.h b/Engine/Source/Developer/SkeletalMeshUtilitiesCommon/Public/LODUtilities.h
--- a/Engine/Source/Developer/SkeletalMeshUtilitiesCommon/Public/LODUtilities.h
+++ b/Engine/Source/Developer/SkeletalMeshUtilitiesCommon/Public/LODUtilities.h
@@ -6,6 +6,7 @@
 
 #include "CoreMinimal.h"
 #include "MeshBuild.h"
+#include "MeshUtilities.h"
 #include "Engine/SkeletalMesh.h"
 #include "Rendering/SkeletalMeshLODImporterData.h"
 #include "Rendering/SkeletalMeshLODModel.h"
@@ -142,9 +143,27 @@ public:
 	 * @param LODIndexDest - the destination LOD
 	 * @param LODIndexSrc - the Source LOD index
 	 */
-	static bool UpdateAlternateSkinWeights(USkeletalMesh* SkeletalMeshDest, const FName& ProfileNameDest, USkeletalMesh* SkeletalMeshSrc, int32 LODIndexDest, int32 LODIndexSrc, FOverlappingThresholds OverlappingThresholds, bool ShouldImportNormals, bool ShouldImportTangents, bool bUseMikkTSpace, bool bComputeWeightedNormals);
+	static bool UpdateAlternateSkinWeights(
+		USkeletalMesh* SkeletalMeshDest,
+		const FName& ProfileNameDest,
+		USkeletalMesh* SkeletalMeshSrc,
+		int32 LODIndexDest,
+		int32 LODIndexSrc,
+		const IMeshUtilities::MeshBuildOptions& Options);
+
+	UE_DEPRECATED(5.2, "Please use the new overloads of UpdateAlternateSkinWeights that take an IMeshUtilities::MeshBuildOptions. Note that IMeshUtilities::MeshBuildOptions::bComputeNormals/Tangents has the opposite meaning of ShouldImportNormals/Tangents.")
+	static bool UpdateAlternateSkinWeights(
+		USkeletalMesh* SkeletalMeshDest,
+		const FName& ProfileNameDest,
+		USkeletalMesh* SkeletalMeshSrc,
+		int32 LODIndexDest,
+		int32 LODIndexSrc,
+		FOverlappingThresholds OverlappingThresholds,
+		bool ShouldImportNormals,
+		bool ShouldImportTangents,
+		bool bUseMikkTSpace,
+		bool bComputeWeightedNormals);
 
-	
 	/*
 	 *	This function apply the skinning weights from the saved imported skinning weight data to the destination skeletal mesh.
 	 *  The Destination will receive the weights has the alternate weights.
@@ -152,12 +171,52 @@ public:
 	 * @param SkeletalMeshDest - The skeletal mesh that will receive the alternate skinning weights.
 	 * @param LODIndexDest - the destination LOD
 	 */
-	static bool UpdateAlternateSkinWeights(USkeletalMesh* SkeletalMeshDest, const FName& ProfileNameDest, int32 LODIndexDest, FOverlappingThresholds OverlappingThresholds, bool ShouldImportNormals, bool ShouldImportTangents, bool bUseMikkTSpace, bool bComputeWeightedNormals);
-
-	static bool UpdateAlternateSkinWeights(FSkeletalMeshLODModel& LODModelDest, FSkeletalMeshImportData& ImportDataDest, USkeletalMesh* SkeletalMeshDest, const FReferenceSkeleton& RefSkeleton, const FName& ProfileNameDest, int32 LODIndexDest, FOverlappingThresholds OverlappingThresholds, bool ShouldImportNormals, bool ShouldImportTangents, bool bUseMikkTSpace, bool bComputeWeightedNormals);
+	static bool UpdateAlternateSkinWeights(
+		USkeletalMesh* SkeletalMeshDest,
+		const FName& ProfileNameDest,
+		int32 LODIndexDest,
+		const IMeshUtilities::MeshBuildOptions& Options);
+
+	UE_DEPRECATED(5.2, "Please use the new overloads of UpdateAlternateSkinWeights that take an IMeshUtilities::MeshBuildOptions. Note that IMeshUtilities::MeshBuildOptions::bComputeNormals/Tangents has the opposite meaning of ShouldImportNormals/Tangents.")
+	static bool UpdateAlternateSkinWeights(
+		USkeletalMesh* SkeletalMeshDest,
+		const FName& ProfileNameDest,
+		int32 LODIndexDest,
+		FOverlappingThresholds OverlappingThresholds,
+		bool ShouldImportNormals,
+		bool ShouldImportTangents,
+		bool bUseMikkTSpace,
+		bool bComputeWeightedNormals);
+
+	static bool UpdateAlternateSkinWeights(
+		FSkeletalMeshLODModel& LODModelDest,
+		FSkeletalMeshImportData& ImportDataDest,
+		USkeletalMesh* SkeletalMeshDest,
+		const FReferenceSkeleton& RefSkeleton,
+		const FName& ProfileNameDest,
+		int32 LODIndexDest,
+		const IMeshUtilities::MeshBuildOptions& Options);
+
+	UE_DEPRECATED(5.2, "Please use the new overloads of UpdateAlternateSkinWeights that take an IMeshUtilities::MeshBuildOptions. Note that IMeshUtilities::MeshBuildOptions::bComputeNormals/Tangents has the opposite meaning of ShouldImportNormals/Tangents.")
+	static bool UpdateAlternateSkinWeights(
+		FSkeletalMeshLODModel& LODModelDest,
+		FSkeletalMeshImportData& ImportDataDest,
+		USkeletalMesh* SkeletalMeshDest,
+		const FReferenceSkeleton& RefSkeleton,
+		const FName& ProfileNameDest,
+		int32 LODIndexDest,
+		FOverlappingThresholds OverlappingThresholds,
+		bool ShouldImportNormals,
+		bool ShouldImportTangents,
+		bool bUseMikkTSpace,
+		bool bComputeWeightedNormals);
 
-	/** Re-generate all (editor-only) skin weight profile, used whenever we rebuild the skeletal mesh data which could change the chunking and bone indices */
-	static void RegenerateAllImportSkinWeightProfileData(FSkeletalMeshLODModel& LODModelDest);
+	/**
+	 * Re-generate all (editor-only) skin weight profile, used whenever we rebuild the skeletal mesh data which could change the chunking and bone indices
+	 * 
+	 * If BoneInfluenceLimit is 0, the DefaultBoneInfluenceLimit from the project settings will be used.
+	 */
+	static void RegenerateAllImportSkinWeightProfileData(FSkeletalMeshLODModel& LODModelDest, int32 BoneInfluenceLimit = 0, const ITargetPlatform* TargetPlatform = nullptr);
 
 	static void UnbindClothingAndBackup(USkeletalMesh* SkeletalMesh, TArray<ClothingAssetUtils::FClothingAssetMeshBinding>& ClothingBindings);
 	static void UnbindClothingAndBackup(USkeletalMesh* SkeletalMesh, TArray<ClothingAssetUtils::FClothingAssetMeshBinding>& ClothingBindings, const int32 LODIndex);
@@ -189,7 +248,11 @@ private:
 	FLODUtilities() {}
 
 	/** Generate the editor-only data stored for a skin weight profile (relies on bone indices) */
-	static void GenerateImportedSkinWeightProfileData(FSkeletalMeshLODModel& LODModelDest, FImportedSkinWeightProfileData& ImportedProfileData);
+	static void GenerateImportedSkinWeightProfileData(
+		FSkeletalMeshLODModel& LODModelDest,
+		struct FImportedSkinWeightProfileData& ImportedProfileData,
+		int32 BoneInfluenceLimit,
+		const class ITargetPlatform* TargetPlatform);
 
 	/**
 	 *	Simplifies the static mesh based upon various user settings for DesiredLOD
diff --git a/Engine/Source/Editor/Persona/Private/PersonaMeshDetails.cpp b/Engine/Source/Editor/Persona/Private/PersonaMeshDetails.cpp
--- a/Engine/Source/Editor/Persona/Private/PersonaMeshDetails.cpp
+++ b/Engine/Source/Editor/Persona/Private/PersonaMeshDetails.cpp
@@ -1708,6 +1708,17 @@ void FSkeletalMeshBuildSettingsLayout::GenerateChildContent(IDetailChildrenBuild
 			FGetFloatDelegate::CreateRaw(this, &FSkeletalMeshBuildSettingsLayout::GetMorphThresholdPosition),
 			FSetFloatDelegate::CreateRaw(this, &FSkeletalMeshBuildSettingsLayout::SetMorphThresholdPosition));
 	}
+
+	{
+		FDetailWidgetRow& BoneInfluenceLimit = AddIntegerRow(ChildrenBuilder,
+			LOCTEXT("BoneInfluenceLimit_Row", "BoneInfluenceLimitRow"),
+			LOCTEXT("BoneInfluenceLimit", "Bone Influence Limit"),
+			LOCTEXT("BoneInfluenceLimit_ToolTip", "Limit the number of bone influences a vertex can have. If 0, the Default Bone Influence Limit from the project settings will be used."),
+			0,
+			MAX_TOTAL_INFLUENCES,
+			FGetIntegerDelegate::CreateRaw(this, &FSkeletalMeshBuildSettingsLayout::GetBoneInfluenceLimit),
+			FSetIntegerDelegate::CreateRaw(this, &FSkeletalMeshBuildSettingsLayout::SetBoneInfluenceLimit));
+	}
 }
 
 bool FSkeletalMeshBuildSettingsLayout::IsBuildEnabled() const
@@ -1782,6 +1793,81 @@ FDetailWidgetRow& FSkeletalMeshBuildSettingsLayout::AddFloatRow(IDetailChildrenB
 	return Row;
 }
 
+FDetailWidgetRow& FSkeletalMeshBuildSettingsLayout::AddIntegerRow(
+	IDetailChildrenBuilder& ChildrenBuilder,
+	const FText& RowTitleText,
+	const FText& RowNameContentText,
+	const FText& RowNameContentTooltipText,
+	const int32 MinSliderValue,
+	const int32 MaxSliderValue,
+	const FGetIntegerDelegate& GetterDelegate,
+	const FSetIntegerDelegate& SetterDelegate)
+{
+	const int32 SliderDataIndex = SliderStateDataArray.Num();
+	FSliderStateData& SliderData = SliderStateDataArray.AddDefaulted_GetRef();
+	SliderData.bSliderActiveMode = false;
+
+	auto BeginSliderMovementHelperFunc = [GetterDelegate, SliderDataIndex, this]()
+	{
+		check(SliderStateDataArray.IsValidIndex(SliderDataIndex));
+		SliderStateDataArray[SliderDataIndex].bSliderActiveMode = true;
+		SliderStateDataArray[SliderDataIndex].MovementValueInt = GetterDelegate.IsBound() ? GetterDelegate.Execute() : 0;
+	};
+
+	auto EndSliderMovementHelperFunc = [SetterDelegate, SliderDataIndex, this](int32 Value)
+	{
+		check(SliderStateDataArray.IsValidIndex(SliderDataIndex));
+		SliderStateDataArray[SliderDataIndex].bSliderActiveMode = false;
+		SliderStateDataArray[SliderDataIndex].MovementValueInt = 0;
+		SetterDelegate.ExecuteIfBound(Value);
+	};
+
+	auto SetValueHelperFunc = [SetterDelegate, SliderDataIndex, this](int32 Value)
+	{
+		check(SliderStateDataArray.IsValidIndex(SliderDataIndex));
+		if (SliderStateDataArray[SliderDataIndex].bSliderActiveMode)
+		{
+			SliderStateDataArray[SliderDataIndex].MovementValueInt = Value;
+		}
+		else
+		{
+			SetterDelegate.ExecuteIfBound(Value);
+		}
+	};
+
+	auto GetValueHelperFunc = [GetterDelegate, SliderDataIndex, this]()
+	{
+		check(SliderStateDataArray.IsValidIndex(SliderDataIndex));
+		if (SliderStateDataArray[SliderDataIndex].bSliderActiveMode)
+		{
+			return SliderStateDataArray[SliderDataIndex].MovementValueInt;
+		}
+		return GetterDelegate.IsBound() ? GetterDelegate.Execute() : 0;
+	};
+
+	FDetailWidgetRow& Row = ChildrenBuilder.AddCustomRow(RowTitleText)
+	.NameContent()
+	[
+		SNew(STextBlock)
+		.Font(IDetailLayoutBuilder::GetDetailFont())
+		.Text(RowNameContentText)
+		.ToolTipText(RowNameContentTooltipText)
+	]
+	.ValueContent()
+	[
+		SNew(SSpinBox<int32>)
+		.Font(IDetailLayoutBuilder::GetDetailFont())
+		.MinValue(MinSliderValue)
+		.MaxValue(MaxSliderValue)
+		.Value_Lambda(GetValueHelperFunc)
+		.OnValueChanged_Lambda(SetValueHelperFunc)
+		.OnBeginSliderMovement_Lambda(BeginSliderMovementHelperFunc)
+		.OnEndSliderMovement_Lambda(EndSliderMovementHelperFunc)
+		.IsEnabled(this, &FSkeletalMeshBuildSettingsLayout::IsBuildEnabled)
+	];
+	return Row;
+}
+
 float FSkeletalMeshBuildSettingsLayout::GetThresholdPosition() const
 {
 	return BuildSettings.ThresholdPosition;
@@ -1850,6 +1936,23 @@ void FSkeletalMeshBuildSettingsLayout::SetMorphThresholdPosition(float Value)
 	}
 }
 
+int32 FSkeletalMeshBuildSettingsLayout::GetBoneInfluenceLimit() const
+{
+	return BuildSettings.BoneInfluenceLimit;
+}
+
+void FSkeletalMeshBuildSettingsLayout::SetBoneInfluenceLimit(int32 Value)
+{
+	if (BuildSettings.BoneInfluenceLimit != Value)
+	{
+		FText TransactionText = FText::Format(LOCTEXT("PersonaSetBoneInfluenceLimitLOD", "LOD{0} build settings: bone influence limit changed"), LODIndex);
+		FScopedTransaction Transaction(TransactionText);
+		ModifyMeshLODSettingsDelegate.ExecuteIfBound(LODIndex);
+
+		BuildSettings.BoneInfluenceLimit = Value;
+	}
+}
+
 ECheckBoxState FSkeletalMeshBuildSettingsLayout::ShouldRecomputeNormals() const
 {
 	return BuildSettings.bRecomputeNormals ? ECheckBoxState::Checked : ECheckBoxState::Unchecked;
diff --git a/Engine/Source/Editor/Persona/Private/PersonaMeshDetails.h b/Engine/Source/Editor/Persona/Private/PersonaMeshDetails.h
--- a/Engine/Source/Editor/Persona/Private/PersonaMeshDetails.h
+++ b/Engine/Source/Editor/Persona/Private/PersonaMeshDetails.h
@@ -255,6 +255,9 @@ public:
 	DECLARE_DELEGATE_RetVal(float, FGetFloatDelegate);
 	DECLARE_DELEGATE_OneParam(FSetFloatDelegate, float);
 
+	DECLARE_DELEGATE_RetVal(int32, FGetIntegerDelegate);
+	DECLARE_DELEGATE_OneParam(FSetIntegerDelegate, int32);
+
 	void UnbindBuildSettings()
 	{
 		IsBuildSettingsEnabledDelegate.Unbind();
@@ -272,8 +275,17 @@ private:
 
 	bool IsBuildEnabled() const;
 
-		//Custom Row Add utilities
+	//Custom Row Add utilities
 	FDetailWidgetRow& AddFloatRow(IDetailChildrenBuilder& ChildrenBuilder, const FText RowTitleText, const FText RowNameContentText, const FText RowNameContentTootlipText, const float MinSliderValue, const float MaxSliderValue, FGetFloatDelegate GetterDelegate, FSetFloatDelegate SetterDelegate);
+	FDetailWidgetRow& AddIntegerRow(
+		IDetailChildrenBuilder& ChildrenBuilder,
+		const FText& RowTitleText,
+		const FText& RowNameContentText,
+		const FText& RowNameContentTooltipText,
+		const int32 MinSliderValue,
+		const int32 MaxSliderValue,
+		const FGetIntegerDelegate& GetterDelegate,
+		const FSetIntegerDelegate& SetterDelegate);
 
 	float GetThresholdPosition() const;
 	void SetThresholdPosition(float Value);
@@ -287,6 +299,9 @@ private:
 	float GetMorphThresholdPosition() const;
 	void SetMorphThresholdPosition(float Value);
 
+	int32 GetBoneInfluenceLimit() const;
+	void SetBoneInfluenceLimit(int32 Value);
+
 	ECheckBoxState ShouldRecomputeNormals() const;
 	ECheckBoxState ShouldRecomputeTangents() const;
 	ECheckBoxState ShouldUseMikkTSpace() const;
diff --git a/Engine/Source/Editor/UnrealEd/Private/SkinWeightsUtilities.cpp b/Engine/Source/Editor/UnrealEd/Private/SkinWeightsUtilities.cpp
--- a/Engine/Source/Editor/UnrealEd/Private/SkinWeightsUtilities.cpp
+++ b/Engine/Source/Editor/UnrealEd/Private/SkinWeightsUtilities.cpp
@@ -254,12 +254,16 @@ bool FSkinWeightsUtilities::ImportAlternateSkinWeight(USkeletalMesh* SkeletalMes
 
 				FImportedSkinWeightProfileData PreviousProfileData = ProfileData;
 				
-				FOverlappingThresholds OverlappingThresholds = ImportOptions.OverlappingThresholds;
-				bool ShouldImportNormals = ImportOptions.ShouldImportNormals();
-				bool ShouldImportTangents = ImportOptions.ShouldImportTangents();
-				bool bUseMikkTSpace = ImportOptions.NormalGenerationMethod == EFBXNormalGenerationMethod::MikkTSpace;
-
-				bResult = FLODUtilities::UpdateAlternateSkinWeights(SkeletalMesh, ProfileName, TmpSkeletalMesh, TargetLODIndex, SrcLodIndex, OverlappingThresholds, ShouldImportNormals, ShouldImportTangents, bUseMikkTSpace, ImportOptions.bComputeWeightedNormals);
+				IMeshUtilities::MeshBuildOptions BuildOptions;
+				BuildOptions.OverlappingThresholds = ImportOptions.OverlappingThresholds;
+				BuildOptions.bComputeNormals = !ImportOptions.ShouldImportNormals();
+				BuildOptions.bComputeTangents = !ImportOptions.ShouldImportTangents();
+				BuildOptions.bUseMikkTSpace = ImportOptions.NormalGenerationMethod == EFBXNormalGenerationMethod::MikkTSpace;
+				BuildOptions.bComputeWeightedNormals = ImportOptions.bComputeWeightedNormals;
+				// There's currently no import option for this. We could add one if needed.
+				BuildOptions.BoneInfluenceLimit = 0;
+
+				bResult = FLODUtilities::UpdateAlternateSkinWeights(SkeletalMesh, ProfileName, TmpSkeletalMesh, TargetLODIndex, SrcLodIndex, BuildOptions);
 				
 				if (!bResult)
 				{
@@ -380,6 +384,8 @@ bool FSkinWeightsUtilities::RemoveSkinnedWeightProfileData(USkeletalMesh* Skelet
 	BuildOptions.bComputeWeightedNormals = OriginalSkeletalMeshImportData->bComputeWeightedNormals;
 	BuildOptions.bRemoveDegenerateTriangles = false;
 	BuildOptions.TargetPlatform = GetTargetPlatformManagerRef().GetRunningTargetPlatform();
+	// There's currently no import option for this. We could add one if needed.
+	BuildOptions.BoneInfluenceLimit = 0;
 
 	//Build the skeletal mesh asset
 	IMeshUtilities& MeshUtilities = FModuleManager::Get().LoadModuleChecked<IMeshUtilities>("MeshUtilities");
@@ -391,7 +397,7 @@ bool FSkinWeightsUtilities::RemoveSkinnedWeightProfileData(USkeletalMesh* Skelet
 
 	//Build the destination mesh with the Alternate influences, so the chunking is done properly.
 	const bool bBuildSuccess = MeshUtilities.BuildSkeletalMesh(LODModelDest, SkeletalMesh->GetPathName(), SkeletalMesh->GetRefSkeleton(), LODInfluencesDest, LODWedgesDest, LODFacesDest, LODPointsDest, LODPointToRawMapDest, BuildOptions, &WarningMessages, &WarningNames);
-	FLODUtilities::RegenerateAllImportSkinWeightProfileData(LODModelDest);
+	FLODUtilities::RegenerateAllImportSkinWeightProfileData(LODModelDest, BuildOptions.BoneInfluenceLimit, BuildOptions.TargetPlatform);
 
 	return bBuildSuccess;
 }
diff --git a/Engine/Source/Runtime/Engine/Classes/Engine/EngineTypes.h b/Engine/Source/Runtime/Engine/Classes/Engine/EngineTypes.h
--- a/Engine/Source/Runtime/Engine/Classes/Engine/EngineTypes.h
+++ b/Engine/Source/Runtime/Engine/Classes/Engine/EngineTypes.h
@@ -2635,6 +2635,16 @@ struct FSkeletalMeshBuildSettings
 	UPROPERTY(EditAnywhere, BlueprintReadWrite, Category = BuildSettings)
 	float MorphThresholdPosition;
 
+	/**
+	 * The maximum number of bone influences to allow each vertex in this mesh to use.
+	 * 
+	 * If set higher than the limit determined by the project settings, it has no effect.
+	 * 
+	 * If set to 0, the value is taken from the DefaultBoneInfluenceLimit project setting.
+	 */
+	UPROPERTY(EditAnywhere, BlueprintReadWrite, Category = BuildSettings)
+	int32 BoneInfluenceLimit;
+
 	/** Default settings. */
 	FSkeletalMeshBuildSettings()
 		: bRecomputeNormals(true)
@@ -2650,6 +2660,7 @@ struct FSkeletalMeshBuildSettings
 		, ThresholdTangentNormal(0.00002f)
 		, ThresholdUV(0.0009765625f)
 		, MorphThresholdPosition(0.015f)
+		, BoneInfluenceLimit(0)
 	{}
 
 	/** Equality operator. */
@@ -2667,7 +2678,8 @@ struct FSkeletalMeshBuildSettings
 			&& ThresholdPosition == Other.ThresholdPosition
 			&& ThresholdTangentNormal == Other.ThresholdTangentNormal
 			&& ThresholdUV == Other.ThresholdUV
-			&& MorphThresholdPosition == Other.MorphThresholdPosition;
+			&& MorphThresholdPosition == Other.MorphThresholdPosition
+			&& BoneInfluenceLimit == Other.BoneInfluenceLimit;
 	}
 
 	/** Inequality. */
diff --git a/Engine/Source/Runtime/Engine/Classes/Engine/RendererSettings.h b/Engine/Source/Runtime/Engine/Classes/Engine/RendererSettings.h
--- a/Engine/Source/Runtime/Engine/Classes/Engine/RendererSettings.h
+++ b/Engine/Source/Runtime/Engine/Classes/Engine/RendererSettings.h
@@ -1096,6 +1096,12 @@ class ENGINE_API URendererSettings : public UDeveloperSettings
 		ToolTip = "When Unlimited Bone Influence is enabled, it still uses a fixed bone inflence buffer until the max bone influence of a mesh exceeds this value"))
 		int32 UnlimitedBonInfluencesThreshold;
 
+	UPROPERTY(config, EditAnywhere, Category = Skinning, meta = (
+		ToolTip = "When BoneInfluenceLimit on a skeletal mesh LOD is set to 0, this setting is used instead. If this setting is 0, no limit will be applied here and the max bone influences will be determined by other project settings. Changing this setting requires restarting the editor.",
+		ConfigRestartRequired = true,
+		ClampMin = "0", UIMin = "0"))
+		FPerPlatformInt DefaultBoneInfluenceLimit;
+
 	/*
 	 * The maximum bones count section vertices's skinning can use before being chunked into more sections. The minimum value is the maximum total influences define (MAX_TOTAL_INFLUENCES).
 	 */
diff --git a/Engine/Source/Runtime/Engine/Private/GPUSkinVertexFactory.cpp b/Engine/Source/Runtime/Engine/Private/GPUSkinVertexFactory.cpp
--- a/Engine/Source/Runtime/Engine/Private/GPUSkinVertexFactory.cpp
+++ b/Engine/Source/Runtime/Engine/Private/GPUSkinVertexFactory.cpp
@@ -458,6 +458,42 @@ bool FGPUBaseSkinVertexFactory::GetUnlimitedBoneInfluences()
 	return (GCVarUnlimitedBoneInfluences!=0);
 }
 
+int32 FGPUBaseSkinVertexFactory::GetBoneInfluenceLimitForAsset(int32 AssetProvidedLimit, const ITargetPlatform* TargetPlatform /*= nullptr*/)
+{
+	if (AssetProvidedLimit > 0)
+	{
+		// The asset provided an explicit limit
+		return AssetProvidedLimit;
+	}
+
+	int32 GlobalDefaultLimit = GetDefault<URendererSettings>()->DefaultBoneInfluenceLimit.GetValue();
+
+#if WITH_EDITOR
+	const ITargetPlatform* TargetPlatformTmp = TargetPlatform;
+	if (!TargetPlatformTmp)
+	{
+		// Get the running platform if the caller did not supply a platform
+		ITargetPlatformManagerModule& TargetPlatformManager = GetTargetPlatformManagerRef();
+		TargetPlatformTmp = TargetPlatformManager.GetRunningTargetPlatform();
+	}
+
+	if (TargetPlatformTmp)
+	{
+		// Get the platform value
+		GlobalDefaultLimit = GetDefault<URendererSettings>()->DefaultBoneInfluenceLimit.GetValueForPlatform(*TargetPlatformTmp->IniPlatformName());
+	}
+#endif
+
+	if (GlobalDefaultLimit > 0)
+	{
+		// A global default limit has been set for this platform
+		return GlobalDefaultLimit;
+	}
+
+	// No limit has been set. Return the maximum possible value.
+	return MAX_TOTAL_INFLUENCES;
+}
+
 void FGPUBaseSkinVertexFactory::SetData(const FGPUSkinDataType* InData)
 {
 	check(InData);
diff --git a/Engine/Source/Runtime/Engine/Private/RendererSettings.cpp b/Engine/Source/Runtime/Engine/Private/RendererSettings.cpp
--- a/Engine/Source/Runtime/Engine/Private/RendererSettings.cpp
+++ b/Engine/Source/Runtime/Engine/Private/RendererSettings.cpp
@@ -86,6 +86,7 @@ URendererSettings::URendererSettings(const FObjectInitializer& ObjectInitializer
 	bEnableRayTracingSkylight = 0;
 	bEnablePathTracing = 0;
 	bEnableRayTracingTextureLOD = 0;
+	DefaultBoneInfluenceLimit = 0;
 	MaxSkinBones = FGPUBaseSkinVertexFactory::GHardwareMaxGPUSkinBones;
 	WorkingColorSpaceChoice = EWorkingColorSpace::sRGB;
 	RedChromaticityCoordinate = FVector2D::ZeroVector;
diff --git a/Engine/Source/Runtime/Engine/Private/SkeletalMesh.cpp b/Engine/Source/Runtime/Engine/Private/SkeletalMesh.cpp
--- a/Engine/Source/Runtime/Engine/Private/SkeletalMesh.cpp
+++ b/Engine/Source/Runtime/Engine/Private/SkeletalMesh.cpp
@@ -4517,6 +4517,9 @@ FString USkeletalMesh::BuildDerivedDataKey(const ITargetPlatform* TargetPlatform
 	const bool bUnlimitedBoneInfluences = FGPUBaseSkinVertexFactory::GetUnlimitedBoneInfluences();
 	KeySuffix += bUnlimitedBoneInfluences ? "1" : "0";
 
+	// Include the global default bone influences limit in case any LODs don't set an explicit limit (highly likely)
+	KeySuffix += FString::FromInt(GetDefault<URendererSettings>()->DefaultBoneInfluenceLimit.GetValueForPlatform(*TargetPlatform->IniPlatformName()));
+
 #if PLATFORM_CPU_ARM_FAMILY
 	// Separate out arm keys as x64 and arm64 clang do not generate the same data for a given
 	// input. Add the arm specifically so that a) we avoid rebuilding the current DDC and
diff --git a/Engine/Source/Runtime/Engine/Private/SkinnedAssetCommon.cpp b/Engine/Source/Runtime/Engine/Private/SkinnedAssetCommon.cpp
--- a/Engine/Source/Runtime/Engine/Private/SkinnedAssetCommon.cpp
+++ b/Engine/Source/Runtime/Engine/Private/SkinnedAssetCommon.cpp
@@ -73,6 +73,7 @@ static void SerializeBuildSettingsForDDC(FArchive& Ar, FSkeletalMeshBuildSetting
 	Ar << BuildSettings.ThresholdTangentNormal;
 	Ar << BuildSettings.ThresholdUV;
 	Ar << BuildSettings.MorphThresholdPosition;
+	Ar << BuildSettings.BoneInfluenceLimit;
 }
 
 FGuid FSkeletalMeshLODInfo::ComputeDeriveDataCacheKey(const FSkeletalMeshLODGroupSettings* SkeletalMeshLODGroupSettings)
diff --git a/Engine/Source/Runtime/Engine/Public/GPUSkinVertexFactory.h b/Engine/Source/Runtime/Engine/Public/GPUSkinVertexFactory.h
--- a/Engine/Source/Runtime/Engine/Public/GPUSkinVertexFactory.h
+++ b/Engine/Source/Runtime/Engine/Public/GPUSkinVertexFactory.h
@@ -378,6 +378,17 @@ public:
 	ENGINE_API static bool UseUnlimitedBoneInfluences(uint32 MaxBoneInfluences);
 	ENGINE_API static bool GetUnlimitedBoneInfluences();
 
+	/*
+	 * Returns the maximum number of bone influences that should be used for a skeletal mesh, given
+	 * the user-requested limit.
+	 * 
+	 * If the requested limit is 0, the limit will be determined from the project settings.
+	 * 
+	 * The return value is guaranteed to be greater than zero, but note that it may be higher than
+	 * the maximum supported bone influences.
+	 */
+	ENGINE_API static int32 GetBoneInfluenceLimitForAsset(int32 AssetProvidedLimit, const ITargetPlatform* TargetPlatform = nullptr);
+
 	/** Morph vertex factory functions */
 	virtual void UpdateMorphVertexStream(const class FMorphVertexBuffer* MorphVertexBuffer) {}
 	virtual const class FMorphVertexBuffer* GetMorphVertexBuffer(bool bPrevious) const { return nullptr; }
