Update skeletal mesh conversion to support geometry scripting.

#rnx

diff --git a/Engine/Plugins/Experimental/Animation/SkeletalMeshModelingTools/Source/SkeletalMeshModelingTools/Private/SkeletalMeshModelingToolsMeshConverter.cpp b/Engine/Plugins/Experimental/Animation/SkeletalMeshModelingTools/Source/SkeletalMeshModelingTools/Private/SkeletalMeshModelingToolsMeshConverter.cpp
--- a/Engine/Plugins/Experimental/Animation/SkeletalMeshModelingTools/Source/SkeletalMeshModelingTools/Private/SkeletalMeshModelingToolsMeshConverter.cpp
+++ b/Engine/Plugins/Experimental/Animation/SkeletalMeshModelingTools/Source/SkeletalMeshModelingTools/Private/SkeletalMeshModelingToolsMeshConverter.cpp
@@ -36,7 +36,15 @@ UObject* USkeletonFromStaticMeshFactory::FactoryCreateNew(
 	)
 {
 	static const FVector RootBoneRelativePosition(0.5, 0.5, 0.0);
-	return FStaticToSkeletalMeshConverter::CreateSkeletonFromStaticMesh(InParent, InName, InFlags, StaticMesh, RootBoneRelativePosition);
+
+	USkeleton* Skeleton = NewObject<USkeleton>(InParent, InName, InFlags);
+	
+	if (!FStaticToSkeletalMeshConverter::InitializeSkeletonFromStaticMesh(Skeleton, StaticMesh, RootBoneRelativePosition))
+	{
+		return nullptr;
+	}
+
+	return Skeleton;
 }
 
 
@@ -60,14 +68,21 @@ UObject* USkeletalMeshFromStaticMeshFactory::FactoryCreateNew(
 	FFeedbackContext* InWarn
 	)
 {
-	USkeletalMesh* SkeletalMesh = FStaticToSkeletalMeshConverter::CreateSkeletalMeshFromStaticMesh(
-		InParent, InName, InFlags, StaticMesh, ReferenceSkeleton, BindBoneName);
-
+	USkeletalMesh* SkeletalMesh = NewObject<USkeletalMesh>(InParent, InName, InFlags);
+	
+	if (!FStaticToSkeletalMeshConverter::InitializeSkeletalMeshFromStaticMesh(SkeletalMesh, StaticMesh, ReferenceSkeleton, BindBoneName))
+	{
+		return nullptr;
+	}
+	
 	// Update the skeletal mesh and the skeleton so that their ref skeletons are in sync and the skeleton's preview mesh
 	// is the one we just created.
 	SkeletalMesh->SetSkeleton(Skeleton);
-	Skeleton->SetPreviewMesh(SkeletalMesh);
 	Skeleton->MergeAllBonesToBoneTree(SkeletalMesh);
+	if (!Skeleton->GetPreviewMesh())
+	{
+		Skeleton->SetPreviewMesh(SkeletalMesh);
+	}
 	
 	return SkeletalMesh;
 }
@@ -94,6 +109,23 @@ USkeleton* UStaticMeshToSkeletalMeshConvertOptions::GetSkeleton(
 	}
 }
 
+
+void UStaticMeshToSkeletalMeshConvertOptions::PostEditChangeProperty(FPropertyChangedEvent& PropertyChangedEvent)
+{
+	Super::PostEditChangeProperty(PropertyChangedEvent);
+
+	if (PropertyChangedEvent.Property)
+	{
+		const FName PropertyName = PropertyChangedEvent.Property->GetFName();
+		if (PropertyName == GET_MEMBER_NAME_CHECKED(UStaticMeshToSkeletalMeshConvertOptions, Skeleton) ||
+			PropertyName == GET_MEMBER_NAME_CHECKED(UStaticMeshToSkeletalMeshConvertOptions, SkeletalMesh))
+		{
+			(void)SkeletonProviderChanged.ExecuteIfBound();
+		}
+	}
+}
+
+
 static bool ConvertSingleMeshToSkeletalMeshInteractive(
 	UStaticMesh* InStaticMesh,
 	TArray<UObject*>& OutObjectsAdded
@@ -102,6 +134,17 @@ static bool ConvertSingleMeshToSkeletalMeshInteractive(
 	UStaticMeshToSkeletalMeshConvertOptions* Options = NewObject<UStaticMeshToSkeletalMeshConvertOptions>();
 	FPropertyEditorModule& PropertyEditorModule = FModuleManager::GetModuleChecked<FPropertyEditorModule>("PropertyEditor");
 
+	// If the object paths don't resolve properly, the details view will still show something, which will be invalidated
+	// later.
+	if (Cast<USkeleton>(Options->Skeleton.ResolveObject()) == nullptr)
+	{
+		Options->Skeleton.Reset();
+	}
+	if (Cast<USkeletalMesh>(Options->SkeletalMesh.ResolveObject()) == nullptr)
+	{
+		Options->SkeletalMesh.Reset();
+	}
+
 	FDetailsViewArgs DetailsViewArgs;
 	DetailsViewArgs.bAllowSearch = false;
 	DetailsViewArgs.NameAreaSettings = FDetailsViewArgs::HideNameArea;
@@ -119,6 +162,11 @@ static bool ConvertSingleMeshToSkeletalMeshInteractive(
 		SCustomDialog::FButton(LOCTEXT("DialogButtonCancel", "Cancel"))
 	});
 
+	Options->SkeletonProviderChanged.BindLambda([DetailsView]()
+	{
+		DetailsView->ForceRefresh();
+	});
+
 	if (OptionsDialog->ShowModal() != 0)
 	{
 		return false;
diff --git a/Engine/Plugins/Experimental/Animation/SkeletalMeshModelingTools/Source/SkeletalMeshModelingTools/Private/SkeletalMeshModelingToolsMeshConverter.h b/Engine/Plugins/Experimental/Animation/SkeletalMeshModelingTools/Source/SkeletalMeshModelingTools/Private/SkeletalMeshModelingToolsMeshConverter.h
--- a/Engine/Plugins/Experimental/Animation/SkeletalMeshModelingTools/Source/SkeletalMeshModelingTools/Private/SkeletalMeshModelingToolsMeshConverter.h
+++ b/Engine/Plugins/Experimental/Animation/SkeletalMeshModelingTools/Source/SkeletalMeshModelingTools/Private/SkeletalMeshModelingToolsMeshConverter.h
@@ -93,17 +93,22 @@ public:
 	UPROPERTY(EditAnywhere, Config, Category="Import")
 	EReferenceSkeletonImportOption SkeletonImportOption = EReferenceSkeletonImportOption::CreateNew;
 
-	UPROPERTY(EditAnywhere, Config, Category="Import", Meta=(AllowedClasses="/Script/Engine.Skeleton", EditCondition="ImportOption==EReferenceSkeletonImportOption::UseExistingSkeleton", EditConditionHides))
+	UPROPERTY(EditAnywhere, Config, Category="Import", Meta=(AllowedClasses="/Script/Engine.Skeleton", EditCondition="SkeletonImportOption==EReferenceSkeletonImportOption::UseExistingSkeleton", EditConditionHides))
 	FSoftObjectPath Skeleton;
 
-	UPROPERTY(EditAnywhere, Config, Category="Import", Meta=(AllowedClasses="/Script/Engine.SkeletalMesh", EditCondition="ImportOption==EReferenceSkeletonImportOption::UseExistingSkeletalMesh", EditConditionHides))
+	UPROPERTY(EditAnywhere, Config, Category="Import", Meta=(AllowedClasses="/Script/Engine.SkeletalMesh", EditCondition="SkeletonImportOption==EReferenceSkeletonImportOption::UseExistingSkeletalMesh", EditConditionHides))
 	FSoftObjectPath SkeletalMesh;
 
-	UPROPERTY(EditAnywhere, Config, Category="Import", meta=(EditCondition="ImportOption!=EReferenceSkeletonImportOption::CreateNew", EditConditionHides))
+	UPROPERTY(EditAnywhere, Config, Category="Import", meta=(EditCondition="SkeletonImportOption!=EReferenceSkeletonImportOption::CreateNew", EditConditionHides))
 	FBoneReference BindingBoneName;
 
 	// IBoneReferenceSkeletonProvider implementation.
 	virtual USkeleton* GetSkeleton(bool& bInvalidSkeletonIsError, const IPropertyHandle* PropertyHandle) override;
+	
+	// UObject overrides.
+	virtual void PostEditChangeProperty(FPropertyChangedEvent& PropertyChangedEvent) override;
+
+	FSimpleDelegate SkeletonProviderChanged;
 };
 
 
diff --git a/Engine/Source/Developer/SkeletalMeshUtilitiesCommon/Private/StaticToSkeletalMeshConverter.cpp b/Engine/Source/Developer/SkeletalMeshUtilitiesCommon/Private/StaticToSkeletalMeshConverter.cpp
--- a/Engine/Source/Developer/SkeletalMeshUtilitiesCommon/Private/StaticToSkeletalMeshConverter.cpp
+++ b/Engine/Source/Developer/SkeletalMeshUtilitiesCommon/Private/StaticToSkeletalMeshConverter.cpp
@@ -25,53 +25,69 @@ static const FName RootBoneName("Root");
 static const TCHAR* JointBaseName(TEXT("Joint"));
 
 
-USkeleton* FStaticToSkeletalMeshConverter::CreateSkeletonFromStaticMesh(
-	UObject *InOuter,
-	const FName InName,
-	const EObjectFlags InFlags,
+bool FStaticToSkeletalMeshConverter::InitializeSkeletonFromStaticMesh(
+	USkeleton* InSkeleton,
 	const UStaticMesh* InStaticMesh,
 	const FVector& InRelativeRootPosition
 	)
 {
-	if (!ensure(InStaticMesh))
+	if (!ensure(InSkeleton))
 	{
-		return nullptr;
+		return false;
 	}
 	
+	if (InSkeleton->GetReferenceSkeleton().GetNum() == 0)
+	{
+		UE_LOG(LogStaticToSkeletalMeshConverter, Error, TEXT("Skeleton '%s' is not empty"), *InSkeleton->GetPathName());
+		return false;
+	}
+	
+	if (!ensure(InStaticMesh))
+	{
+		return false;
+	}
+
 	const FBox Bounds = InStaticMesh->GetBoundingBox();
 	const FVector RootPosition = Bounds.Min + (Bounds.Max - Bounds.Min) * InRelativeRootPosition;
 	FTransform RootTransform(FTransform::Identity);
 	RootTransform.SetTranslation(RootPosition);
 
-	USkeleton* Skeleton = NewObject<USkeleton>(InOuter, InName, InFlags);
-
-	FReferenceSkeletonModifier Modifier(Skeleton);
+	FReferenceSkeletonModifier Modifier(InSkeleton);
 	Modifier.Add(FMeshBoneInfo(RootBoneName, RootBoneName.ToString(), INDEX_NONE), RootTransform);
 
-	return Skeleton;
+	return true;
 }
 
 
-USkeleton* FStaticToSkeletalMeshConverter::CreateSkeletonFromStaticMesh(
-	UObject *InOuter,
-	const FName InName,
-	const EObjectFlags InFlags,
+bool FStaticToSkeletalMeshConverter::InitializeSkeletonFromStaticMesh(
+	USkeleton* InSkeleton,
 	const UStaticMesh* InStaticMesh,
 	const FVector& InRelativeRootPosition,
 	const FVector& InRelativeEndEffectorPosition,
 	const int32 InIntermediaryJointCount
 	)
 {
-	if (FMath::IsNearlyZero(FVector::DistSquared(InRelativeEndEffectorPosition, InRelativeRootPosition)))
+	if (!ensure(InSkeleton))
+	{
+		return false;
+	}
+	
+	if (InSkeleton->GetReferenceSkeleton().GetNum() == 0)
 	{
-		return CreateSkeletonFromStaticMesh(InOuter, InName, InFlags, InStaticMesh, InRelativeRootPosition);
+		UE_LOG(LogStaticToSkeletalMeshConverter, Error, TEXT("Skeleton '%s' is not empty"), *InSkeleton->GetPathName());
+		return false;
 	}
 
 	if (!ensure(InStaticMesh))
 	{
-		return nullptr;
+		return false;
 	}
 	
+	if (FMath::IsNearlyZero(FVector::DistSquared(InRelativeEndEffectorPosition, InRelativeRootPosition)))
+	{
+		return InitializeSkeletonFromStaticMesh(InSkeleton, InStaticMesh, InRelativeRootPosition);
+	}
+
 	const FBox Bounds = InStaticMesh->GetBoundingBox();
 	const FVector RootPosition = Bounds.Min + (Bounds.Max - Bounds.Min) * InRelativeRootPosition;
 	const FVector EndEffectorPosition = Bounds.Min + (Bounds.Max - Bounds.Min) * InRelativeEndEffectorPosition;
@@ -83,8 +99,7 @@ USkeleton* FStaticToSkeletalMeshConverter::CreateSkeletonFromStaticMesh(
 	ParentTransform.SetTranslation(RootPosition);
 	ParentTransform.SetRotation(Rotation);
 
-	USkeleton* Skeleton = NewObject<USkeleton>(InOuter, InName, InFlags);
-	FReferenceSkeletonModifier Modifier(Skeleton);
+	FReferenceSkeletonModifier Modifier(InSkeleton);
 	Modifier.Add(FMeshBoneInfo(RootBoneName, RootBoneName.ToString(), INDEX_NONE), ParentTransform);
 
 	for (int32 JointIndex = 0; JointIndex <= InIntermediaryJointCount; JointIndex++)
@@ -98,7 +113,7 @@ USkeleton* FStaticToSkeletalMeshConverter::CreateSkeletonFromStaticMesh(
 		ParentTransform = PointTransform;
 	}
 
-	return Skeleton;
+	return true;
 }
 
 static void CopyBuildSettings(
@@ -168,9 +183,57 @@ static void CopyReductionSettings(
 
 
 static bool AddLODFromMeshDescription(
+	const FMeshDescription& InMeshDescription,
+	USkeletalMesh* InSkeletalMesh,
+	IMeshUtilities& InMeshUtilities
+	)
+{
+	FSkeletalMeshModel* ImportedModels = InSkeletalMesh->GetImportedModel();
+	const int32 LODIndex = ImportedModels->LODModels.Num(); 
+	ImportedModels->LODModels.Add(new FSkeletalMeshLODModel);
+	if (!ensure(ImportedModels->LODModels.Num() == InSkeletalMesh->GetLODNum()))
+	{
+		return false;
+	}
+	
+	FSkeletalMeshLODModel& SkeletalMeshModel = ImportedModels->LODModels.Last(); 
+	
+	FSkeletalMeshImportData SkeletalMeshImportGeometry = FSkeletalMeshImportData::CreateFromMeshDescription(InMeshDescription);
+
+	// We need at least one set of texture coordinates. Always.
+	SkeletalMeshModel.NumTexCoords = FMath::Max<uint32>(1, SkeletalMeshImportGeometry.NumTexCoords);
+	
+	// Data needed by BuildSkeletalMesh
+	TArray<FVector3f> LODPoints;
+	TArray<SkeletalMeshImportData::FMeshWedge> LODWedges;
+	TArray<SkeletalMeshImportData::FMeshFace> LODFaces;
+	TArray<SkeletalMeshImportData::FVertInfluence> LODInfluences;
+	TArray<int32> LODPointToRawMap;
+	SkeletalMeshImportGeometry.CopyLODImportData( LODPoints, LODWedges, LODFaces, LODInfluences, LODPointToRawMap );
+
+	IMeshUtilities::MeshBuildOptions BuildOptions;
+	BuildOptions.TargetPlatform = GetTargetPlatformManagerRef().GetRunningTargetPlatform();
+	BuildOptions.FillOptions(InSkeletalMesh->GetLODInfoArray().Last().BuildSettings);
+
+	TArray<FText> WarningMessages;
+	if (!InMeshUtilities.BuildSkeletalMesh(SkeletalMeshModel, InSkeletalMesh->GetPathName(), InSkeletalMesh->GetRefSkeleton(), LODInfluences, LODWedges, LODFaces, LODPoints, LODPointToRawMap, BuildOptions, &WarningMessages, nullptr))
+	{
+		for(const FText& Message: WarningMessages)
+		{
+			UE_LOG(LogStaticToSkeletalMeshConverter, Warning, TEXT("%s"), *Message.ToString());
+		}
+		return false;
+	}
+	
+	InSkeletalMesh->SaveLODImportedData(LODIndex, SkeletalMeshImportGeometry);
+	
+	return true;
+}
+
+
+static bool AddLODFromStaticMeshSourceModel(
 	const FStaticMeshSourceModel& InStaticMeshSourceModel,
 	USkeletalMesh* InSkeletalMesh,
-	const int32 InLODIndex,
 	const FBoneIndexType InBoneIndex,
 	IMeshUtilities& InMeshUtilities
 	)
@@ -182,10 +245,6 @@ static bool AddLODFromMeshDescription(
 	CopyBuildSettings(InStaticMeshSourceModel.BuildSettings, SkeletalLODInfo.BuildSettings);
 	CopyReductionSettings(InStaticMeshSourceModel.ReductionSettings, SkeletalLODInfo.ReductionSettings);
 
-	FSkeletalMeshModel* ImportedModels = InSkeletalMesh->GetImportedModel();
-	ImportedModels->LODModels.Add(new FSkeletalMeshLODModel);
-	FSkeletalMeshLODModel& SkeletalMeshModel = ImportedModels->LODModels.Last(); 
-	
 	if (InStaticMeshSourceModel.IsMeshDescriptionValid())
 	{
 		FMeshDescription SkeletalMeshGeometry;
@@ -207,101 +266,88 @@ static bool AddLODFromMeshDescription(
 			SkinWeights.Set(VertexID, RootBinding);
 		}
 
-		FSkeletalMeshImportData SkeletalMeshImportGeometry = FSkeletalMeshImportData::CreateFromMeshDescription(SkeletalMeshGeometry);
-
-		// We need at least one set of texture coordinates. Always.
-		SkeletalMeshModel.NumTexCoords = FMath::Max<uint32>(1, SkeletalMeshImportGeometry.NumTexCoords);
-		
-		// Data needed by BuildSkeletalMesh
-		TArray<FVector3f> LODPoints;
-		TArray<SkeletalMeshImportData::FMeshWedge> LODWedges;
-		TArray<SkeletalMeshImportData::FMeshFace> LODFaces;
-		TArray<SkeletalMeshImportData::FVertInfluence> LODInfluences;
-		TArray<int32> LODPointToRawMap;
-		SkeletalMeshImportGeometry.CopyLODImportData( LODPoints, LODWedges, LODFaces, LODInfluences, LODPointToRawMap );
-
-		IMeshUtilities::MeshBuildOptions BuildOptions;
-		BuildOptions.TargetPlatform = GetTargetPlatformManagerRef().GetRunningTargetPlatform();
-		BuildOptions.FillOptions(SkeletalLODInfo.BuildSettings);
-
-		TArray<FText> WarningMessages;
-		if (!InMeshUtilities.BuildSkeletalMesh(SkeletalMeshModel, InSkeletalMesh->GetPathName(), InSkeletalMesh->GetRefSkeleton(), LODInfluences, LODWedges, LODFaces, LODPoints, LODPointToRawMap, BuildOptions, &WarningMessages, nullptr))
+		if (!AddLODFromMeshDescription(SkeletalMeshGeometry, InSkeletalMesh, InMeshUtilities))
 		{
-			for(const FText& Message: WarningMessages)
-			{
-				UE_LOG(LogStaticToSkeletalMeshConverter, Warning, TEXT("%s"), *Message.ToString());
-			}
 			return false;
 		}
-		
-		InSkeletalMesh->SaveLODImportedData(InLODIndex, SkeletalMeshImportGeometry);
 	}
 	else
 	{
+		FSkeletalMeshModel* ImportedModels = InSkeletalMesh->GetImportedModel();
+		const int32 LODIndex = ImportedModels->LODModels.Num(); 
+		ImportedModels->LODModels.Add(new FSkeletalMeshLODModel);
+			
 		FSkeletalMeshUpdateContext UpdateContext;
 		UpdateContext.SkeletalMesh = InSkeletalMesh;
 		
-		FLODUtilities::SimplifySkeletalMeshLOD(UpdateContext, InLODIndex, GetTargetPlatformManagerRef().GetRunningTargetPlatform());
+		FLODUtilities::SimplifySkeletalMeshLOD(UpdateContext, LODIndex, GetTargetPlatformManagerRef().GetRunningTargetPlatform());
 	}
 	
 	return true;
 }
 
 
-
-USkeletalMesh* FStaticToSkeletalMeshConverter::CreateSkeletalMeshFromStaticMesh(
-	UObject *InOuter,
-	const FName InName,
-	const EObjectFlags InFlags,
+bool FStaticToSkeletalMeshConverter::InitializeSkeletalMeshFromStaticMesh(
+	USkeletalMesh* InSkeletalMesh,
 	const UStaticMesh* InStaticMesh,
 	const FReferenceSkeleton& InReferenceSkeleton,
 	const FName InBindBone 
 )
 {
+	if (!ensure(InSkeletalMesh))
+	{
+		return false;
+	}
+
+	if (!InSkeletalMesh->GetImportedModel()->LODModels.IsEmpty())
+	{
+		UE_LOG(LogStaticToSkeletalMeshConverter, Error, TEXT("Skeletal mesh '%s' is not empty"), *InSkeletalMesh->GetPathName());
+		return false;
+	}
+
 	if (!ensure(InStaticMesh))
 	{
-		return nullptr;
+		return false;
 	}
 
 	int32 BoneIndex = 0;
 	if (!InBindBone.IsNone())
 	{
 		BoneIndex = InReferenceSkeleton.FindRawBoneIndex(InBindBone);
-		if (!ensure(BoneIndex != INDEX_NONE))
+		if (BoneIndex == INDEX_NONE)
 		{
-			return nullptr;
+			UE_LOG(LogStaticToSkeletalMeshConverter, Error, TEXT("Bone '%s' not found in skeleton."), *InBindBone.ToString());
+			return false;
 		}
 	}
 
-	USkeletalMesh* SkeletalMesh = NewObject<USkeletalMesh>(InOuter, InName, InFlags);
-
 	// This ensures that the render data gets built before we return, by calling PostEditChange when we fall out of scope.
-	FScopedSkeletalMeshPostEditChange ScopedPostEditChange( SkeletalMesh );
-	SkeletalMesh->PreEditChange( nullptr );
-	SkeletalMesh->SetRefSkeleton(InReferenceSkeleton);
+	FScopedSkeletalMeshPostEditChange ScopedPostEditChange( InSkeletalMesh );
+	InSkeletalMesh->PreEditChange( nullptr );
+	InSkeletalMesh->SetRefSkeleton(InReferenceSkeleton);
 	
 	// Calculate the initial pose from the reference skeleton.
-	SkeletalMesh->CalculateInvRefMatrices();
+	InSkeletalMesh->CalculateInvRefMatrices();
 
 	IMeshUtilities& MeshUtilities = FModuleManager::Get().LoadModuleChecked<IMeshUtilities>( "MeshUtilities" );
 	
 	// Copy the LODs and LOD settings over (as close as we can).
-	for (int32 Index = 0; Index < InStaticMesh->GetNumSourceModels(); Index++)
+	bool bFirstSourceModel = true;
+	for (const FStaticMeshSourceModel& StaticMeshSourceModel: InStaticMesh->GetSourceModels())
 	{
-		const FStaticMeshSourceModel& StaticMeshSourceModel = InStaticMesh->GetSourceModel(Index);
-		
-		if (!AddLODFromMeshDescription(
-			StaticMeshSourceModel, SkeletalMesh, Index, static_cast<FBoneIndexType>(BoneIndex), MeshUtilities))
+		if (!AddLODFromStaticMeshSourceModel(
+			StaticMeshSourceModel, InSkeletalMesh, static_cast<FBoneIndexType>(BoneIndex), MeshUtilities))
 		{
 			// If we didn't get a model for LOD index 0, we don't have a mesh. Bail out.
-			if (Index == 0)
+			if (bFirstSourceModel)
 			{
-				return nullptr;
+				return false;
 			}
 
 			// Otherwise, we have a model, so let's continue with what we have.
 			break;
 		}
+		bFirstSourceModel = false;
 	}
 	
 	// Convert the materials over.
@@ -315,14 +361,158 @@ USkeletalMesh* FStaticToSkeletalMeshConverter::CreateSkeletalMeshFromStaticMesh(
 		Materials.Add(Material);
 	}
 
-	SkeletalMesh->SetMaterials(Materials);
+	InSkeletalMesh->SetMaterials(Materials);
 	
 	// Set the bounds from the static mesh, including the extensions, otherwise it won't render properly (among other things).
-	SkeletalMesh->SetImportedBounds( InStaticMesh->GetBounds() );
-	SkeletalMesh->SetPositiveBoundsExtension(InStaticMesh->GetPositiveBoundsExtension());
-	SkeletalMesh->SetNegativeBoundsExtension(InStaticMesh->GetNegativeBoundsExtension());
+	InSkeletalMesh->SetImportedBounds( InStaticMesh->GetBounds() );
+	InSkeletalMesh->SetPositiveBoundsExtension(InStaticMesh->GetPositiveBoundsExtension());
+	InSkeletalMesh->SetNegativeBoundsExtension(InStaticMesh->GetNegativeBoundsExtension());
+
+	return true;
+}
 
-	return SkeletalMesh;
+
+static bool ValidateSkinWeightAttribute(
+	const FMeshDescription& InMeshDescription,
+	const FReferenceSkeleton& InReferenceSkeleton
+	)
+{
+	using namespace UE::AnimationCore;
+	FSkeletalMeshConstAttributes MeshAttributes{InMeshDescription};
+
+	TArray<FName> Profiles = MeshAttributes.GetSkinWeightProfileNames(); 
+	if (Profiles.IsEmpty())
+	{
+		UE_LOG(LogStaticToSkeletalMeshConverter, Error, TEXT("Mesh description doesn't have a skin weight attribute."));
+		return false;
+	}
+
+	FBoneIndexType BoneIndexMax = static_cast<FBoneIndexType>(InReferenceSkeleton.GetRawBoneNum());
+		
+	// We use the first profile. Usually that's the default profile, unless we have nothing but alternate profiles.
+	FSkinWeightsVertexAttributesConstRef VertexSkinWeights = MeshAttributes.GetVertexSkinWeights(Profiles[0]);
+	for (const FVertexID VertexID: InMeshDescription.Vertices().GetElementIDs())
+	{
+		for (FBoneWeight BoneWeight: VertexSkinWeights.Get(VertexID))
+		{
+			if (BoneWeight.GetBoneIndex() >= BoneIndexMax)
+			{
+				UE_LOG(LogStaticToSkeletalMeshConverter, Error, TEXT("Mesh description's skin weight refers to a non-existent bone (%d of %d)."), BoneWeight.GetBoneIndex(), BoneIndexMax);
+				return false;
+			}
+		}
+	}
+	return true;
+}
+
+
+bool FStaticToSkeletalMeshConverter::InitializeSkeletalMeshFromMeshDescriptions(
+	USkeletalMesh* InSkeletalMesh,
+	TArrayView<const FMeshDescription*> InMeshDescriptions,
+	TConstArrayView<FSkeletalMaterial> InMaterials,
+	const FReferenceSkeleton& InReferenceSkeleton,
+	const bool bInRecomputeNormals,
+	const bool bInRecomputeTangents
+)	
+{
+	if (!ensure(InSkeletalMesh))
+	{
+		return false;
+	}
+
+	if (!InSkeletalMesh->GetImportedModel()->LODModels.IsEmpty())
+	{
+		UE_LOG(LogStaticToSkeletalMeshConverter, Error, TEXT("Skeletal mesh '%s' is not empty"), *InSkeletalMesh->GetPathName());
+		return false;
+	}
+
+	if (InMeshDescriptions.IsEmpty())
+	{
+		UE_LOG(LogStaticToSkeletalMeshConverter, Error, TEXT("No mesh descriptions given"));
+		return false;
+	}
+
+	// Ensure all mesh descriptions have a skin weight attribute.
+	for (const FMeshDescription* MeshDescription: InMeshDescriptions)
+	{
+		if (!ValidateSkinWeightAttribute(*MeshDescription, InReferenceSkeleton))
+		{
+			return false;
+		}
+	}
+
+	// Set the materials before we start converting. We'll add dummy materials afterward if there are more sections
+	// than materials in any of the LODs. Not the best system, but the best we have for now.
+	InSkeletalMesh->SetMaterials(TArray<FSkeletalMaterial>{InMaterials});
+
+	// This ensures that the render data gets built before we return, by calling PostEditChange when we fall out of scope.
+	{
+		FScopedSkeletalMeshPostEditChange ScopedPostEditChange( InSkeletalMesh );
+		InSkeletalMesh->PreEditChange( nullptr );
+		InSkeletalMesh->SetRefSkeleton(InReferenceSkeleton);
+		
+		// Calculate the initial pose from the reference skeleton.
+		InSkeletalMesh->CalculateInvRefMatrices();
+
+		IMeshUtilities& MeshUtilities = FModuleManager::Get().LoadModuleChecked<IMeshUtilities>( "MeshUtilities" );
+		bool bFirstSourceModel = true;
+		
+		for (const FMeshDescription* MeshDescription: InMeshDescriptions)
+		{
+			// Add default LOD build settings.
+			FSkeletalMeshLODInfo& SkeletalLODInfo = InSkeletalMesh->AddLODInfo();
+
+			SkeletalLODInfo.BuildSettings.bRecomputeNormals = bInRecomputeNormals;
+			SkeletalLODInfo.BuildSettings.bRecomputeTangents = bInRecomputeTangents;
+
+			if (!AddLODFromMeshDescription(*MeshDescription, InSkeletalMesh, MeshUtilities))
+			{
+				// If we didn't get a model for LOD index 0, we don't have a mesh. Bail out.
+				if (bFirstSourceModel)
+				{
+					return false;
+				}
+
+				// Otherwise, we have a model, so let's continue with what we have.
+				break;
+			}
+			bFirstSourceModel = false;
+		}
+	}
+
+	// Compute the bbox, now that we have the model mesh generated. 
+	FBox3f BoundingBox{ForceInit};
+	int32 MaxSectionCount = 0;
+	for (const FSkeletalMeshLODModel& MeshModel: InSkeletalMesh->GetImportedModel()->LODModels)
+	{
+		MaxSectionCount = FMath::Max(MaxSectionCount, MeshModel.Sections.Num());
+
+		// Compute the overall bbox.
+		for (const FSkelMeshSection& Section: MeshModel.Sections)
+		{
+			for (const FSoftSkinVertex& Vertex: Section.SoftVertices)
+			{
+				BoundingBox += Vertex.Position;
+			}
+		}
+	}
+
+	// If we're short on materials, compared to sections, add dummy materials to fill in the gap. Not ideal, but
+	// best we can do for now.
+	const TArray<FSkeletalMaterial>& ExistingMaterials = InSkeletalMesh->GetMaterials();
+	if (MaxSectionCount > ExistingMaterials.Num())
+	{
+		TArray<FSkeletalMaterial> NewMaterials{ExistingMaterials};
+		for (int32 Index = ExistingMaterials.Num(); Index < MaxSectionCount; Index++)
+		{
+			NewMaterials.Add(FSkeletalMaterial{});
+		}
+		InSkeletalMesh->SetMaterials(NewMaterials);
+	}
+
+	InSkeletalMesh->SetImportedBounds( FBox3d{BoundingBox} );
+
+	return true;
 }
 
 #endif
diff --git a/Engine/Source/Developer/SkeletalMeshUtilitiesCommon/Public/StaticToSkeletalMeshConverter.h b/Engine/Source/Developer/SkeletalMeshUtilitiesCommon/Public/StaticToSkeletalMeshConverter.h
--- a/Engine/Source/Developer/SkeletalMeshUtilitiesCommon/Public/StaticToSkeletalMeshConverter.h
+++ b/Engine/Source/Developer/SkeletalMeshUtilitiesCommon/Public/StaticToSkeletalMeshConverter.h
@@ -6,30 +6,30 @@
 
 #include "Math/MathFwd.h"
 #include "UObject/NameTypes.h"
-#include "UObject/ObjectMacros.h"
 
 
 class UObject;
+class USkeletalMesh;
 class USkeleton;
 class UStaticMesh;
-class USkeletalMesh;
+struct FMeshDescription;
 struct FReferenceSkeleton;
+struct FSkeletalMaterial;
 
 
 struct SKELETALMESHUTILITIESCOMMON_API FStaticToSkeletalMeshConverter
 {
 	/** Creates a skeleton from a static mesh with a single root bone.
 	 *
+	 *  @param InSkeleton The skeleton to initialize. It has to be freshly created and contain no bones already.
 	 *  @param InStaticMesh The static mesh whose bounding box will be used as a reference.
 	 *  @param InRelativeRootPosition The relative root position in a unit cube that gets scaled up to match the
 	 *    bbox of the static mesh. For example, given FVector(0.5, 0.5, 0.0), the root bone location will be placed
 	 *    at the center of the bottom of the static mesh's bbox.
-	 *  NOTE: The owner upon return will be the transient package, re-outer to the correct package as needed. 
+	 *  @return \c true if the operation succeeded. 
 	 */
-	static USkeleton* CreateSkeletonFromStaticMesh(
-		UObject *InOuter,
-		const FName InName,
-		const EObjectFlags InFlags,
+	static bool InitializeSkeletonFromStaticMesh(
+		USkeleton* InSkeleton,
 		const UStaticMesh* InStaticMesh,
 		const FVector& InRelativeRootPosition 
 		);
@@ -37,6 +37,7 @@ struct SKELETALMESHUTILITIESCOMMON_API FStaticToSkeletalMeshConverter
 	/** Creates a skeleton from a static mesh with a bone chain going from root to end effector, where the intermediary
 	 *  bones are distributed evenly along a line between the two.
 	 *
+	 *  @param InSkeleton The skeleton to initialize. It has to be freshly created and contain no bones already.
 	 *  @param InStaticMesh The static mesh whose bounding box will be used as a reference.
 	 *  @param InRelativeRootPosition The relative root position in a unit cube that gets scaled up to match the
 	 *    bbox of the static mesh. For example, given FVector(0.5, 0.5, 0.0), the root bone location will be placed
@@ -44,35 +45,49 @@ struct SKELETALMESHUTILITIESCOMMON_API FStaticToSkeletalMeshConverter
 	 *  @param InRelativeEndEffectorPosition The end effector position, positioned in the same manner as the root
 	 *    position. If the end effector is in the same location as the root, only the root bone is created.
 	 *  @param InIntermediaryJointCount Number of joints to create between the root and the end effector.  
-	 *  NOTE: The owner upon return will be the transient package, re-outer to the correct package as needed. 
+	 *  @return \c true if the operation succeeded. 
 	 */
-	static USkeleton* CreateSkeletonFromStaticMesh(
-		UObject *InOuter,
-		const FName InName,
-		const EObjectFlags InFlags,
+	static bool InitializeSkeletonFromStaticMesh(
+		USkeleton* InSkeleton,
 		const UStaticMesh* InStaticMesh,
 		const FVector& InRelativeRootPosition,
 		const FVector& InRelativeEndEffectorPosition,
 		const int32 InIntermediaryJointCount
 		);
 
-	/** Create a skeletal mesh from the given static mesh and a skeleton. The mesh will initially be created with a
-	 *  rigid binding on the root bone. Use SetRigidBinding or SetSmoothBinding to override this given binding.
-	 *  @param InOuter The skeletal mesh's outer/parent object.
-	 *  @param InName The name to give the newly created skeletal mesh.
-	 *  @param InFlags The object creation flags to use.
+	/** Initialize a skeletal mesh from the given static mesh and a skeleton. The mesh will initially be created with a
+	 *  rigid binding on the root bone. 
+	 *  @param InSkeletalMesh The skeletal mesh to initialize. It has to be freshly created and be completely empty.
 	 *  @param InStaticMesh The static mesh to convert from.
 	 *  @param InReferenceSkeleton The reference skeleton to use.
 	 *  @param InBindBone The bone to bind to. If no bone name is given, the binding defaults to the root bone.
 	 */
-	static USkeletalMesh* CreateSkeletalMeshFromStaticMesh(
-		UObject *InOuter,
-		const FName InName,
-		const EObjectFlags InFlags,
+	static bool InitializeSkeletalMeshFromStaticMesh(
+		USkeletalMesh* InSkeletalMesh, 
 		const UStaticMesh* InStaticMesh,
 		const FReferenceSkeleton& InReferenceSkeleton,
 		const FName InBindBone = NAME_None 
 		);
+
+	/** Create a skeletal mesh from a list of mesh description objects and a reference skeleton. If any mesh description
+	 *  either doesn't contain a skin weight, or any of the skin weights reference a bone that doesn't exist on the
+	 *  reference skeleton, this operation fails.
+	 *  @param InSkeletalMesh The skeletal mesh to initialize. It has to be freshly created and be completely empty.
+	 *  @param InMeshDescriptions The mesh description objects to create each subsequent LOD, starting from LOD0.
+	 *  @param InMaterials The materials to apply to the skeletal mesh. The conversion will attempt a best effort to
+	 *     ensure that the PolygonGroup::ImportedMaterialSlotName attribute will map to existing import names. 
+	 *  @param InReferenceSkeleton The reference skeleton to use.
+	 *  @param bInRecomputeNormals Recompute normals when the mesh is generated from the mesh description.
+	 *  @param bInRecomputeTangents Recompute tangents when the mesh is generated from the mesh description.
+	 */
+	static bool InitializeSkeletalMeshFromMeshDescriptions(
+		USkeletalMesh* InSkeletalMesh,
+		TArrayView<const FMeshDescription*> InMeshDescriptions,
+		TConstArrayView<FSkeletalMaterial> InMaterials,
+		const FReferenceSkeleton& InReferenceSkeleton,
+		const bool bInRecomputeNormals = false,
+		const bool bInRecomputeTangents = true
+		);
 };
 
 #endif
