Allow skinned mesh reduction to use parallelfor when running from workers

#rnx

diff --git a/Engine/Source/Developer/SkeletalMeshUtilitiesCommon/Private/LODUtilities.cpp b/Engine/Source/Developer/SkeletalMeshUtilitiesCommon/Private/LODUtilities.cpp
--- a/Engine/Source/Developer/SkeletalMeshUtilitiesCommon/Private/LODUtilities.cpp
+++ b/Engine/Source/Developer/SkeletalMeshUtilitiesCommon/Private/LODUtilities.cpp
@@ -2902,7 +2902,7 @@ void FLODUtilities::RegenerateDependentLODs(USkeletalMesh* SkeletalMesh, int32 L
 				int32 DependentLODIndex = DependentLODs[IterationIndex];
 				check(SkeletalMesh->GetLODInfo(DependentLODIndex)); //We cannot add a LOD when reducing with multi thread, so check we already have one
 				FLODUtilities::SimplifySkeletalMeshLOD(SkeletalMesh, DependentLODIndex, TargetPlatform, false, &bNeedsPackageDirtied);
-			}, IsInGameThread() ? EParallelForFlags::None : EParallelForFlags::ForceSingleThread);
+			});
 
 			if (bNeedsPackageDirtied && IsInGameThread())
 			{
