Fixing minor bug when checking for number of bones in new skeleton when creating a skeletal mesh from a static mesh.

diff --git a/Engine/Source/Developer/SkeletalMeshUtilitiesCommon/Private/StaticToSkeletalMeshConverter.cpp b/Engine/Source/Developer/SkeletalMeshUtilitiesCommon/Private/StaticToSkeletalMeshConverter.cpp
--- a/Engine/Source/Developer/SkeletalMeshUtilitiesCommon/Private/StaticToSkeletalMeshConverter.cpp
+++ b/Engine/Source/Developer/SkeletalMeshUtilitiesCommon/Private/StaticToSkeletalMeshConverter.cpp
@@ -36,7 +36,7 @@ bool FStaticToSkeletalMeshConverter::InitializeSkeletonFromStaticMesh(
 		return false;
 	}
 	
-	if (InSkeleton->GetReferenceSkeleton().GetNum() == 0)
+	if (InSkeleton->GetReferenceSkeleton().GetNum() != 0)
 	{
 		UE_LOG(LogStaticToSkeletalMeshConverter, Error, TEXT("Skeleton '%s' is not empty"), *InSkeleton->GetPathName());
 		return false;
@@ -72,7 +72,7 @@ bool FStaticToSkeletalMeshConverter::InitializeSkeletonFromStaticMesh(
 		return false;
 	}
 	
-	if (InSkeleton->GetReferenceSkeleton().GetNum() == 0)
+	if (InSkeleton->GetReferenceSkeleton().GetNum() != 0)
 	{
 		UE_LOG(LogStaticToSkeletalMeshConverter, Error, TEXT("Skeleton '%s' is not empty"), *InSkeleton->GetPathName());
 		return false;
