Fix assert when showing alternate skinning
#rnx

diff --git a/Engine/Source/Developer/SkeletalMeshUtilitiesCommon/Private/LODUtilities.cpp b/Engine/Source/Developer/SkeletalMeshUtilitiesCommon/Private/LODUtilities.cpp
--- a/Engine/Source/Developer/SkeletalMeshUtilitiesCommon/Private/LODUtilities.cpp
+++ b/Engine/Source/Developer/SkeletalMeshUtilitiesCommon/Private/LODUtilities.cpp
@@ -2768,6 +2768,11 @@ void FLODUtilities::GenerateImportedSkinWeightProfileData(
 			InfluenceBones[InfluenceBoneCount] = Kvp.Key;
 			InfluenceWeights[InfluenceBoneCount] = Kvp.Value;
 			InfluenceBoneCount++;
+			//Do not add more influence then the maximum we support on the platform or for this asset
+			if (InfluenceBoneCount >= MaxInfluenceCount)
+			{
+				break;
+			}
 		}
 
 		const FBoneWeights BoneWeights = FBoneWeights::Create(InfluenceBones, InfluenceWeights, InfluenceBoneCount);
@@ -2800,6 +2805,7 @@ void FLODUtilities::GenerateImportedSkinWeightProfileData(
 				Section.MaxBoneInfluences = MaxNumInfluences;
 			}
 		}
+		check(MaxNumInfluences <= MaxInfluenceCount);
 	}
 }
 
