Skeletal mesh - skeleton new editing API

diff --git a/Engine/Source/Developer/SkeletalMeshUtilitiesCommon/Private/SkeletonModifier.cpp b/Engine/Source/Developer/SkeletalMeshUtilitiesCommon/Private/SkeletonModifier.cpp
new file mode 100644
--- /dev/null
+++ b/Engine/Source/Developer/SkeletalMeshUtilitiesCommon/Private/SkeletonModifier.cpp
@@ -0,0 +1,1134 @@
+﻿// Copyright Epic Games, Inc. All Rights Reserved.
+
+
+#include "SkeletonModifier.h"
+
+#include "BoneWeights.h"
+#include "LODUtilities.h"
+#include "MeshDescription.h"
+#include "Animation/Skeleton.h"
+#include "ReferenceSkeleton.h"
+#include "SkeletalMeshAttributes.h"
+#include "Animation/AnimSequence.h"
+#include "Animation/PoseAsset.h"
+#include "Animation/AnimMontage.h"
+#include "AssetRegistry/AssetRegistryModule.h"
+#include "AssetRegistry/AssetData.h"
+#include "Rendering/SkeletalMeshModel.h"
+
+FTransform FMirrorOptions::MirrorTransform(const FTransform& InTransform) const
+{
+	FTransform Transform = InTransform;
+	Transform.SetLocation(MirrorVector(Transform.GetLocation()));
+
+	if (bMirrorRotation)
+	{
+		FRotator Rotator(ForceInitToZero);
+		switch (MirrorAxis)
+		{
+		case EAxis::X:
+			{
+				Rotator.Roll = 180;
+				break;
+			}
+		case EAxis::Y:
+			{
+				Rotator.Pitch = 180;
+				break;
+			}
+		case EAxis::Z:
+			{
+				Rotator.Yaw = 180;
+				break;
+			}
+		default:
+			break;
+		}	
+		Transform.SetRotation(FQuat::MakeFromRotator(Rotator) * Transform.GetRotation());
+	}
+
+	return Transform;
+}
+
+FVector FMirrorOptions::MirrorVector(const FVector& InVector) const
+{
+	FVector Axis = FVector::ZeroVector;
+	Axis.SetComponentForAxis(MirrorAxis, 1.f);
+	return InVector.MirrorByVector(Axis);
+}
+
+void FSkeletonModifier::ExternalUpdate(const FReferenceSkeleton& InRefSkeleton, const TArray<int32>& InIndexTracker)
+{
+	if (!ReferenceSkeleton)
+	{
+		return;
+	}
+	
+	*ReferenceSkeleton = InRefSkeleton;
+	TransformComposer.Reset(new FTransformComposer(*ReferenceSkeleton));
+	BoneIndexTracker = InIndexTracker;
+}
+
+bool FSkeletonModifier::Init(USkeletalMesh* InSkeletalMesh)
+{
+	SkeletalMesh = nullptr;
+	MeshDescription.Reset();
+	ReferenceSkeleton.Reset();
+	TransformComposer.Reset();
+	BoneIndexTracker.Reset();
+
+#if WITH_EDITORONLY_DATA
+	// validate supplied skeletal mesh exists
+	if (!InSkeletalMesh)
+	{
+		UE_LOG(LogAnimation, Error, TEXT("Skeleton Modifier: No skeletal mesh supplied to load."));
+		return false;
+	}
+
+	const USkeleton* Skeleton = InSkeletalMesh->GetSkeleton();
+	if (!Skeleton)
+	{
+		UE_LOG(LogAnimation, Error, TEXT("Skeleton Modifier: Skeletal Mesh supplied has no skeleton."));
+		return false;
+	}
+
+	// check assets using this skeleton ?
+	{
+		static const TArray<FTopLevelAssetPath> AssetPaths({
+		   UAnimSequence::StaticClass()->GetClassPathName(),
+		   UAnimMontage::StaticClass()->GetClassPathName(),
+		   UPoseAsset::StaticClass()->GetClassPathName(),
+		   USkeletalMesh::StaticClass()->GetClassPathName()});
+
+		FARFilter Filter;
+		Filter.ClassPaths = AssetPaths;
+
+		const IAssetRegistry& AssetRegistry = FModuleManager::LoadModuleChecked<FAssetRegistryModule>("AssetRegistry").Get();
+		TArray<FAssetData> Assets; AssetRegistry.GetAssets(Filter, Assets);
+
+		const FAssetData SkeletonAssetData(Skeleton);
+		const FString SkeletonPath = SkeletonAssetData.GetExportTextName();
+	
+		static const FName Tag("Skeleton");
+		for (const FAssetData& AssetData: Assets)
+		{
+			const FString TagValue = AssetData.GetTagValueRef<FString>(Tag);
+			if (TagValue == SkeletonPath)
+			{
+				UE_LOG(LogAnimation, Warning, TEXT("%s references that skeleton."), *AssetData.GetExportTextName());
+			}
+		}
+	}
+
+	// verify mesh has an imported model
+	const FSkeletalMeshModel* SkeletalMeshModel = InSkeletalMesh->GetImportedModel();
+	if (!SkeletalMeshModel)
+	{
+		UE_LOG(LogAnimation, Error, TEXT("Skeleton Modifier: Supplied skeletal mesh does not have an imported model."));
+		return false;
+	}
+	
+	// verify that we have a LOD 0 mesh
+	if (SkeletalMeshModel->LODModels.IsEmpty())
+	{
+		UE_LOG(LogAnimation, Error, TEXT("Skeleton Modifier: Supplied skeletal mesh does not have a LOD 0 mesh."));
+		return false;
+	}
+	
+	// verify user is not trying to modify one of the core engine assets
+	if (InSkeletalMesh->GetPathName().StartsWith(TEXT("/Engine/")))
+	{
+		UE_LOG(LogAnimation, Error, TEXT("Skeleton Modifier: Cannot modify built-in engine asset."));
+		return false;
+	}
+	
+	// store pointer to mesh and instantiate a mesh description for commiting changes
+	SkeletalMesh = InSkeletalMesh;
+	
+	MeshDescription = MakeUnique<FMeshDescription>();
+	const FSkeletalMeshLODModel& LOD0Mesh = SkeletalMeshModel->LODModels[0];
+	LOD0Mesh.GetMeshDescription(*MeshDescription, SkeletalMesh);
+
+	// store reference skeleton to edit
+	ReferenceSkeleton = MakeUnique<FReferenceSkeleton>();
+	*ReferenceSkeleton = SkeletalMesh->GetRefSkeleton();
+
+	TransformComposer.Reset(new FTransformComposer(*ReferenceSkeleton));
+
+	// store initial bones indices to track for changes
+	const int32 NumBones = ReferenceSkeleton->GetRawBoneNum();
+	BoneIndexTracker.Reserve(NumBones);
+	for (int32 Index = 0; Index < NumBones; Index++)
+	{
+		BoneIndexTracker.Add(Index);
+	}
+	return true;
+#else
+	ensureMsgf(false, TEXT("FSkeletonModifier is an editor only feature."));
+#endif
+	return false;
+}
+
+bool FSkeletonModifier::IsReferenceSkeletonValid(const bool bLog) const
+{
+	if (!ReferenceSkeleton.IsValid())
+	{
+		if(bLog)
+		{
+			UE_LOG(LogAnimation, Error, TEXT("Skeleton Modifier: No valid reference skeleton provided."));
+		}
+		return false;		
+	}
+	return true;
+}
+
+bool FSkeletonModifier::CommitSkeletonToSkeletalMesh()
+{
+	if (!SkeletalMesh || !ReferenceSkeleton || !MeshDescription)
+	{
+		UE_LOG(LogAnimation, Error, TEXT("Skeleton Modifier: No mesh loaded. Cannot apply skeleton edits."));
+		return false;
+	}
+
+#if WITH_EDITORONLY_DATA
+	// before commiting, we have to reparent non-root bones with no parent as the animation pipeline
+	// doesn't support multi-roots
+	const TArray<FMeshBoneInfo>& BoneInfos = ReferenceSkeleton->GetRawRefBoneInfo();
+	if (!BoneInfos.IsEmpty())
+	{
+		TArray<FName> BonesToParent;
+		for (int32 Index = 1; Index < BoneInfos.Num(); Index++)
+		{
+			const FMeshBoneInfo& BoneInfo = BoneInfos[Index];
+			if (BoneInfo.ParentIndex == INDEX_NONE)
+			{
+				BonesToParent.Add(BoneInfo.Name);
+			}
+		}
+		
+		if (!BonesToParent.IsEmpty())
+		{
+			for (const FName& BoneName: BonesToParent)
+			{
+				UE_LOG(LogAnimation, Warning, TEXT("Skeleton Modifier: %s will be parented to the root bone before commiting."), *BoneName.ToString());
+			}
+			ParentBones(BonesToParent, {BoneInfos[0].Name});
+		}
+	}
+	
+	// update skeletal mesh
+	SkeletalMesh->SetRefSkeleton(*ReferenceSkeleton);
+	SkeletalMesh->GetRefBasesInvMatrix().Reset();
+	SkeletalMesh->CalculateInvRefMatrices(); 
+	
+	// update skeleton
+	USkeleton* Skeleton = SkeletalMesh->GetSkeleton();
+	if (Skeleton->RecreateBoneTree(SkeletalMesh))
+	{
+		Skeleton->MarkPackageDirty();	
+	}
+
+	// update mesh description
+	{
+		FSkeletalMeshAttributes MeshAttributes(*MeshDescription);
+
+		// update bone data
+		if (!MeshAttributes.HasBones())
+		{
+			MeshAttributes.Register(true);
+		}
+
+		MeshAttributes.Bones().Reset(BoneInfos.Num());
+
+		FSkeletalMeshAttributes::FBoneNameAttributesRef BoneNames = MeshAttributes.GetBoneNames();
+		FSkeletalMeshAttributes::FBoneParentIndexAttributesRef BoneParentIndices = MeshAttributes.GetBoneParentIndices();
+		FSkeletalMeshAttributes::FBonePoseAttributesRef BonePoses = MeshAttributes.GetBonePoses();
+
+		const TArray<FTransform> Transforms = ReferenceSkeleton->GetRawRefBonePose();
+		for (int Index = 0; Index < BoneInfos.Num(); ++Index)
+		{
+			const FMeshBoneInfo& Info = BoneInfos[Index];
+			const FBoneID BoneID = MeshAttributes.CreateBone();
+			BoneNames.Set(BoneID, Info.Name);
+			BoneParentIndices.Set(BoneID, Info.ParentIndex);
+			BonePoses.Set(BoneID, Transforms[Index]);
+		}
+		
+		// update skin data if needed
+		auto HasBoneIndexesChanged = [&]()
+		{
+			for (int32 Index = 0; Index < BoneIndexTracker.Num(); Index++)
+			{
+				if (BoneIndexTracker[Index] != Index)
+				{
+					return true;
+				}
+			}
+			return false;
+		};
+		if (HasBoneIndexesChanged())
+		{
+			using namespace UE::AnimationCore;
+			FBoneWeightsSettings BoneSettings; BoneSettings.SetNormalizeType(EBoneWeightNormalizeType::None);
+		
+			FSkinWeightsVertexAttributesRef SkinWeights = MeshAttributes.GetVertexSkinWeights();
+			for (const FVertexID& VertexID: MeshDescription->Vertices().GetElementIDs())
+			{
+				FVertexBoneWeights BoneWeights = SkinWeights.Get(VertexID);
+				if (const int32 NumBoneWeights = BoneWeights.Num())
+				{
+					TArray<FBoneWeight> NewWeights;
+					for (int32 Idx = 0; Idx < NumBoneWeights; ++Idx)
+					{
+						const FBoneWeight& OldBoneWeight = BoneWeights[Idx];
+						const int32 BoneIndex = OldBoneWeight.GetBoneIndex();
+					
+						check(BoneIndexTracker.IsValidIndex(BoneIndex));
+					
+						const int32 NewBoneIndex = BoneIndexTracker[BoneIndex];
+						if (NewBoneIndex != INDEX_NONE)
+						{
+							NewWeights.Add(FBoneWeight(NewBoneIndex, OldBoneWeight.GetRawWeight()));
+						}
+					}
+					SkinWeights.Set(VertexID, FBoneWeights::Create(NewWeights, BoneSettings));
+				}
+			}
+		}
+	}
+
+	// update skeletal mesh LOD (cf. USkeletalMeshToolTarget::CommitMeshDescription)
+	{
+		FlushRenderingCommands();
+		SkeletalMesh->SetFlags(RF_Transactional);
+		SkeletalMesh->Modify();
+		// generate mesh import data from the mesh description (now with modified weights)
+		FSkeletalMeshImportData SkeletalMeshImportData = FSkeletalMeshImportData::CreateFromMeshDescription(*MeshDescription);
+		SkeletalMesh->SaveLODImportedData(0, SkeletalMeshImportData);
+		SkeletalMesh->SetLODImportedDataVersions(0, ESkeletalMeshGeoImportVersions::LatestVersion, ESkeletalMeshSkinningImportVersions::LatestVersion);
+		SkeletalMesh->SetUseLegacyMeshDerivedDataKey(false);
+		SkeletalMesh->PostEditChange();
+	}
+
+	return true;
+#else
+	ensureMsgf(false, TEXT("FSkeletonModifier is an editor only feature."));
+#endif
+	return false;
+}
+
+bool FSkeletonModifier::AddBone(const FName InBoneName, const FName InParentName, const FTransform& InTransform)
+{
+	if (InBoneName == NAME_None)
+	{
+		UE_LOG(LogAnimation, Error, TEXT("Skeleton Modifier - Add: Cannot add bone with no name."));
+		return false;
+	}
+	
+	return AddBones({InBoneName}, {InParentName}, {InTransform});
+}
+
+bool FSkeletonModifier::AddBones(
+	const TArray<FName>& InBoneNames, const TArray<FName>& InParentNames, const TArray<FTransform>& InTransforms)
+{
+	if (!IsReferenceSkeletonValid())
+	{
+		return false;
+	}
+	
+	const int32 NumBonesToAdd = InBoneNames.Num();
+	if (NumBonesToAdd == 0)
+	{
+		UE_LOG(LogAnimation, Error, TEXT("Skeleton Modifier - Add: The provided bone names array is empty."));
+		return false;
+	}
+	
+	struct FBoneData
+	{
+		FMeshBoneInfo BoneInfo;
+		int32 TransformOffset;
+	};
+
+	TArray<FBoneData> BonesToAdd;
+	BonesToAdd.Reserve(NumBonesToAdd);
+
+	auto GetParentName = [&](int32 Index)
+	{
+		if (InBoneNames.Num() == InParentNames.Num())
+		{
+			return InParentNames[Index];
+		}
+		return InParentNames.IsEmpty() ? NAME_None : InParentNames[0];  
+	};
+
+	const int NumBonesBefore = ReferenceSkeleton->GetRawBoneNum();
+
+	for (int32 Index = 0; Index < NumBonesToAdd; Index++)
+	{
+		const FName& BoneName = InBoneNames[Index];		
+		if (ReferenceSkeleton->FindBoneIndex(InBoneNames[Index]) == INDEX_NONE)
+		{
+			const FName ParentName = GetParentName(Index);
+
+			// look for parent index in the ref skeleton
+			int32 ParentIndex = ReferenceSkeleton->FindBoneIndex(ParentName);
+			if (ParentIndex == INDEX_NONE && Index > 0)
+			{
+				// otherwise, check if one of the new bone is going to be the parent
+				ParentIndex = InBoneNames.IndexOfByKey(ParentName);
+
+				if (ParentIndex > INDEX_NONE && ParentIndex < Index)
+				{
+					ParentIndex += NumBonesBefore;
+				}
+			}
+			const FMeshBoneInfo NewBoneInfo(BoneName, BoneName.ToString(), ParentIndex);
+			BonesToAdd.Add({NewBoneInfo, Index});
+		}
+	}
+
+	if (BonesToAdd.IsEmpty())
+	{
+		UE_LOG(LogAnimation, Error, TEXT("Skeleton Modifier - Add: None of the provided names is avalable to be added."));
+		return false;
+	}
+
+	auto GetTransform = [&](int32 Index)
+	{
+		if (InBoneNames.Num() == InTransforms.Num())
+		{
+			return InTransforms[Index];
+		}
+		return InTransforms.IsEmpty() ? FTransform::Identity : InTransforms[0];  
+	};
+	
+	//update reference skeleton
+	{
+		static constexpr bool bAllowMultipleRoots = true;
+		FReferenceSkeletonModifier Modifier(*ReferenceSkeleton, nullptr);
+		for (FBoneData& BoneData: BonesToAdd)
+		{
+			Modifier.Add(BoneData.BoneInfo, GetTransform(BoneData.TransformOffset), bAllowMultipleRoots);
+		}
+	}
+
+	// invalidate composer
+	TransformComposer->Invalidate(INDEX_NONE);
+	
+	// update index tracker: nothing to do as those new indices do not represent any bone in the initial skinning data
+	
+	return true;
+}
+
+bool FSkeletonModifier::MirrorBone(const FName InBoneName, const FMirrorOptions& InOptions)
+{
+	if (InBoneName == NAME_None)
+	{
+		UE_LOG(LogAnimation, Error, TEXT("Skeleton Modifier - Mirror: Cannot mirror bone with no name."));
+		return false;
+	}
+	
+	return MirrorBones({InBoneName}, InOptions);
+}
+
+bool FSkeletonModifier::MirrorBones(const TArray<FName>& InBonesName, const FMirrorOptions& InOptions)
+{
+	if (!IsReferenceSkeletonValid())
+	{
+		return false;
+	}
+
+	// get bones to mirror
+	TArray<int32> BonesToMirror; GetBonesToMirror(InBonesName, InOptions, BonesToMirror);
+	
+	const int32 NumBonesToMirror = BonesToMirror.Num();
+	if (NumBonesToMirror == 0)
+	{
+		UE_LOG(LogAnimation, Error, TEXT("Skeleton Modifier - Mirror: None of the provided names has been found."));
+		return false;
+	}
+
+	// get mirrored names
+	TArray<FName> MirroredNames; GetMirroredNames(BonesToMirror, InOptions, MirroredNames);
+
+	// add bones first if they are missing
+	TArray<int32> MirroredIndex; GetMirroredBones(BonesToMirror, MirroredNames, MirroredIndex);
+	if (MirroredIndex.Num() != NumBonesToMirror)
+	{
+		UE_LOG(LogAnimation, Error, TEXT("Skeleton Modifier - Mirror: Couldn't find mirrored bones."));
+		return false;
+	}
+
+	// compute mirrored transforms
+	TArray<FTransform> MirroredTransforms; GetMirroredTransforms(BonesToMirror, InOptions, MirroredTransforms);
+
+	// update reference skeleton
+	{
+		FReferenceSkeletonModifier Modifier(*ReferenceSkeleton, nullptr);
+		for (int32 Index = 0; Index < NumBonesToMirror; Index++)
+		{
+			Modifier.UpdateRefPoseTransform(MirroredIndex[Index], MirroredTransforms[Index]);
+		}
+	}
+
+	// invalidate composer
+	TransformComposer->Invalidate(INDEX_NONE);
+	
+	return true;
+}
+
+void FSkeletonModifier::GetBonesToMirror(
+	const TArray<FName>& InBonesName, const FMirrorOptions& InOptions, TArray<int32>& OutBonesToMirror) const
+{
+	OutBonesToMirror.Reset();
+	
+	TSet<int32> IndicesToMirror;
+	auto GetBonesToMirror = [&](const int32 BoneIndex, auto&& GetBonesToMirror2) -> void
+	{
+		if (BoneIndex == INDEX_NONE)
+		{
+			return;
+		}
+		
+		IndicesToMirror.Add(BoneIndex);
+
+		if (InOptions.bMirrorChildren)
+		{
+			TArray<int32> Children; ReferenceSkeleton->GetDirectChildBones(BoneIndex, Children);
+			for (int32 ChildIndex: Children)
+			{
+				GetBonesToMirror2(ChildIndex, GetBonesToMirror2);
+			}
+		}
+	};
+	
+	for (const FName& BoneName: InBonesName)
+	{
+		GetBonesToMirror(ReferenceSkeleton->FindRawBoneIndex(BoneName), GetBonesToMirror);
+	}
+
+	const int32 NumBonesToMirror = IndicesToMirror.Num();
+	if (NumBonesToMirror == 0)
+	{
+		return;
+	}
+
+	IndicesToMirror.Sort([](const int32 Index0, const int32 Index1) {return Index0 < Index1;});
+	OutBonesToMirror = IndicesToMirror.Array();
+}
+
+void FSkeletonModifier::GetMirroredNames(
+	const TArray<int32>& InBonesToMirror, const FMirrorOptions& InOptions, TArray<FName>& OutBonesName) const
+{
+	OutBonesName.Reset();
+	if (InBonesToMirror.IsEmpty())
+	{
+		return;
+	}
+	
+	const TArray<FMeshBoneInfo>& BoneInfos = ReferenceSkeleton->GetRawRefBoneInfo();
+	
+	OutBonesName.Reserve(InBonesToMirror.Num());
+	Algo::Transform(InBonesToMirror, OutBonesName, [&](const int32 BoneIndex)
+	{
+		FString BoneStr = BoneInfos[BoneIndex].Name.ToString();
+		if (BoneStr.Contains(InOptions.LeftString, ESearchCase::CaseSensitive, ESearchDir::FromEnd))
+		{
+			BoneStr = BoneStr.Replace(*InOptions.LeftString, *InOptions.RightString, ESearchCase::CaseSensitive);
+		}
+		else if (BoneStr.Contains(InOptions.RightString, ESearchCase::CaseSensitive, ESearchDir::FromEnd))
+		{
+			BoneStr = BoneStr.Replace(*InOptions.RightString, *InOptions.LeftString, ESearchCase::CaseSensitive);
+		}
+		else
+		{
+			return GetUniqueName(BoneInfos[BoneIndex].Name);
+		}
+		return FName(*BoneStr);
+	});
+}
+
+void FSkeletonModifier::GetMirroredBones(
+	const TArray<int32>& InBonesToMirror, const TArray<FName>& InMirroredNames, TArray<int32>& OutMirroredBones)
+{
+	OutMirroredBones.Reset();
+	if (InBonesToMirror.IsEmpty() || InBonesToMirror.Num() != InMirroredNames.Num())
+	{
+		return;
+	}
+	
+	const TArray<FMeshBoneInfo>& BoneInfos = ReferenceSkeleton->GetRawRefBoneInfo();
+	const TArray<FTransform>& BoneTransforms = ReferenceSkeleton->GetRawRefBonePose();
+	
+	TArray<FName> BonesToAdd, ParentNames;
+	TArray<FTransform> Transforms;
+	const int32 NumBones = InBonesToMirror.Num();
+	for (int32 Index = 0; Index < NumBones; Index++)
+	{
+		const FName& MirroredName = InMirroredNames[Index];
+		const int32 MirroredIndex = ReferenceSkeleton->FindRawBoneIndex(MirroredName);
+		if (MirroredIndex == INDEX_NONE)
+		{
+			const int32 RefBoneIndex = InBonesToMirror[Index];
+				
+			// name
+			BonesToAdd.Add(MirroredName);
+
+			// parent
+			int32 ParentIndex = BoneInfos[RefBoneIndex].ParentIndex;
+			FName ParentName = NAME_None;
+			if (ParentIndex != INDEX_NONE)
+			{
+				ParentName = BoneInfos[ParentIndex].Name;
+				// is that parent being mirrored?
+				const int32 ParentIndexInMirrored = InBonesToMirror.IndexOfByKey(ParentIndex);
+				if (ParentIndexInMirrored != INDEX_NONE)
+				{
+					ParentName = InMirroredNames[ParentIndexInMirrored];
+				}
+			}
+			ParentNames.Add(ParentName);
+
+			// transform
+			Transforms.Add(BoneTransforms[RefBoneIndex]);
+		}
+		else
+		{
+			OutMirroredBones.Add(MirroredIndex);
+		}
+	}
+
+	if (BonesToAdd.IsEmpty())
+	{
+		return;
+	}
+
+	// add missing bones and get their index
+	OutMirroredBones.Reset();
+	if (!AddBones(BonesToAdd, ParentNames, Transforms))
+	{
+		return;
+	}
+
+	OutMirroredBones.Reserve(NumBones);
+	for (int32 Index = 0; Index < NumBones; Index++)
+	{
+		const FName& MirroredName = InMirroredNames[Index];
+		const int32 MirroredIndex = ReferenceSkeleton->FindRawBoneIndex(MirroredName);
+		if (MirroredIndex == INDEX_NONE)
+		{
+			OutMirroredBones.Reset();
+			return;
+		}
+		OutMirroredBones.Add(MirroredIndex);
+	}
+}
+
+void FSkeletonModifier::GetMirroredTransforms(
+	const TArray<int32>& InBonesToMirror, const FMirrorOptions& InOptions, TArray<FTransform>& OutMirroredTransforms) const
+{
+	OutMirroredTransforms.Reset();
+	
+	const int32 NumBonesToMirror = InBonesToMirror.Num();
+	if (NumBonesToMirror == 0)
+	{
+		return;
+	}
+	
+	const TArray<FMeshBoneInfo>& BoneInfos = ReferenceSkeleton->GetRawRefBoneInfo();
+	auto FindFirstNotMirroredParent = [&](const int32 RefBoneIndex) -> int32
+	{
+		if (RefBoneIndex == INDEX_NONE)
+		{
+			return INDEX_NONE;
+		}
+		
+		int32 RefParentIndex = BoneInfos[RefBoneIndex].ParentIndex;
+		if (RefParentIndex == INDEX_NONE)
+		{
+			return INDEX_NONE;
+		}
+
+		int32 ParentMirroredIndex = InBonesToMirror.IndexOfByKey(RefParentIndex);
+		while (ParentMirroredIndex != INDEX_NONE)
+		{
+			RefParentIndex = BoneInfos[RefParentIndex].ParentIndex;
+			ParentMirroredIndex = InBonesToMirror.IndexOfByKey(RefParentIndex);
+		}
+		return RefParentIndex;
+	};
+
+	// compute global mirrored transforms
+	TArray<FTransform> MirroredGlobal; MirroredGlobal.Reserve(NumBonesToMirror);
+	for (int32 Index = 0; Index < NumBonesToMirror; Index++)
+	{
+		// bone global
+		const int32 RefBoneIndex = InBonesToMirror[Index];
+		FTransform Global = TransformComposer->GetGlobalTransform(RefBoneIndex);
+	
+		// first parent not mirrored global
+		const int32 FirstNotMirroredParent = FindFirstNotMirroredParent(RefBoneIndex);
+		const FTransform& FirstNotMirroredGlobal = TransformComposer->GetGlobalTransform(FirstNotMirroredParent);
+
+		// switch to first parent not mirrored (translation only)
+		Global.AddToTranslation(-FirstNotMirroredGlobal.GetTranslation());
+
+		// mirror
+		Global = InOptions.MirrorTransform(Global);
+
+		// switch back to global (translation only)
+		Global.AddToTranslation(FirstNotMirroredGlobal.GetTranslation());
+
+		MirroredGlobal.Add(Global);
+	}
+
+	// switch back to local
+	OutMirroredTransforms.Reserve(NumBonesToMirror);
+	for (int32 Index = 0; Index < NumBonesToMirror; Index++)
+	{
+		const int32 RefBoneIndex = InBonesToMirror[Index];		
+		const int32 RefParentIndex = BoneInfos[RefBoneIndex].ParentIndex;
+		const int32 ParentMirroredIndex = InBonesToMirror.IndexOfByKey(RefParentIndex);
+		const FTransform& ParentGlobal = ParentMirroredIndex != INDEX_NONE ? MirroredGlobal[ParentMirroredIndex] :
+							TransformComposer->GetGlobalTransform(RefParentIndex);
+		OutMirroredTransforms.Add(MirroredGlobal[Index].GetRelativeTransform(ParentGlobal));
+	}
+}
+
+// NOTE : that function might take a bUpdateChildren to decide whether we want to compensate the children transforms
+// atm, we update the bone's local ref transform so children's global transforms are changed (we just need to cache the
+// global transforms then restore them back)
+// orienting the bone for example should change the children's global transform
+bool FSkeletonModifier::MoveBone( const FName InBoneToMove, const FTransform& InNewTransform, const bool bMoveChildren)
+{
+	if (InBoneToMove == NAME_None)
+	{
+		UE_LOG(LogAnimation, Error, TEXT("Skeleton Modifier - Move: Cannot move bone with no name."));
+		return false;
+	}
+	
+	return MoveBones({InBoneToMove}, {InNewTransform}, bMoveChildren);
+}
+
+bool FSkeletonModifier::MoveBones(
+	const TArray<FName>& InBonesToMove, const TArray<FTransform>& InNewTransforms, const bool bMoveChildren)
+{
+	if (!IsReferenceSkeletonValid())
+	{
+		return false;
+	}
+	
+	const int32 NumBonesToMove = InBonesToMove.Num();
+	if (NumBonesToMove == 0 || NumBonesToMove != InNewTransforms.Num())
+	{
+		UE_LOG(LogAnimation, Error, TEXT("Skeleton Modifier - Move: Discrepancy between bones and transforms (%d / %d)."), NumBonesToMove, InNewTransforms.Num());
+		return false;
+	}
+	
+	TArray<int32> BoneIndices, Offsets;
+	BoneIndices.Reserve(NumBonesToMove);
+	Offsets.Reserve(NumBonesToMove);
+
+	for (int32 Index = 0; Index < NumBonesToMove; Index++)
+	{
+		const int32 BoneIndex = ReferenceSkeleton->FindRawBoneIndex(InBonesToMove[Index]);
+		if (BoneIndex != INDEX_NONE)
+		{
+			BoneIndices.Add(BoneIndex);
+			Offsets.Add(Index);
+		}
+	}
+
+	if (BoneIndices.IsEmpty())
+	{
+		UE_LOG(LogAnimation, Error, TEXT("Skeleton Modifier - Move: None of the provided bones has been found."));
+		return false;
+	}
+
+	// compute global transforms if needed
+	TArray<int32> ChildrenToFix;
+	TArray<FTransform> GlobalTransforms;
+	
+	if (!bMoveChildren)
+	{
+		// get children
+		for (int32 Index = 0; Index < BoneIndices.Num(); Index++)
+		{
+			TArray<int32> Children;
+			ReferenceSkeleton->GetDirectChildBones(BoneIndices[Index], Children);
+			ChildrenToFix.Append(Children);
+		}
+
+		// sort them from highest index to lowest
+		ChildrenToFix.Sort([](const int32 Index0, const int32 Index1) {return Index0 > Index1;} );
+		const int32 NumChildren = ChildrenToFix.Num();
+
+		// compute global transforms (note that we could cache them for faster implementation) 
+		GlobalTransforms.AddUninitialized(NumChildren);
+
+		for (int32 Index = 0; Index < NumChildren; Index++)
+		{
+			GlobalTransforms[Index] = TransformComposer->GetGlobalTransform(ChildrenToFix[Index]);
+		}
+	}
+	
+	// update reference skeleton
+	{
+		FReferenceSkeletonModifier Modifier(*ReferenceSkeleton, nullptr);
+		for (int32 Index = 0; Index < BoneIndices.Num(); Index++)
+		{
+			Modifier.UpdateRefPoseTransform(BoneIndices[Index], InNewTransforms[Offsets[Index]]);
+
+			// invalidate cached global transform
+			TransformComposer->Invalidate(BoneIndices[Index]);
+		}
+
+		if (!bMoveChildren)
+		{
+			const int32 NumChildren = ChildrenToFix.Num();
+			for (int32 Index = 0; Index < NumChildren; Index++)
+			{
+				const int32 ChildrenIndex = ChildrenToFix[Index];
+				const int32 ParentIndex = ReferenceSkeleton->GetRawParentIndex(ChildrenIndex);
+				const FTransform& NewParentGlobal = TransformComposer->GetGlobalTransform(ParentIndex);
+				const FTransform& NewLocal = GlobalTransforms[Index].GetRelativeTransform(NewParentGlobal);
+				Modifier.UpdateRefPoseTransform(ChildrenIndex, NewLocal);
+				TransformComposer->Invalidate(ChildrenIndex);
+			}
+		}
+	}
+
+	// update index tracker: no modification on bone indices to track when changing transforms
+	
+	return true;
+}
+
+bool FSkeletonModifier::RemoveBone(const FName InBoneName, const bool bRemoveChildren)
+{
+	if (InBoneName == NAME_None)
+	{
+		UE_LOG(LogAnimation, Error, TEXT("Skeleton Modifier - Remove: Cannot remove bone with no name."));
+		return false;
+	}
+	
+	return RemoveBones({InBoneName}, bRemoveChildren);
+}
+
+bool FSkeletonModifier::RemoveBones(const TArray<FName>& InBoneNames, const bool bRemoveChildren)
+{
+	if (!IsReferenceSkeletonValid())
+	{
+		return false;
+	}
+
+	if (InBoneNames.IsEmpty())
+	{
+		UE_LOG(LogAnimation, Error, TEXT("Skeleton Modifier - Remove: No bone provided."));
+		return false;
+	}
+	
+	// store initial data
+	const TArray<FMeshBoneInfo> InfosBeforeRemoval = ReferenceSkeleton->GetRawRefBoneInfo();
+
+	// update reference skeleton
+	{
+		FReferenceSkeletonModifier Modifier(*ReferenceSkeleton, nullptr);
+		for (const FName& BoneName: InBoneNames)
+		{
+			Modifier.Remove(BoneName, bRemoveChildren);
+		}
+	}
+
+	if (InfosBeforeRemoval.Num() == ReferenceSkeleton->GetRawRefBoneInfo().Num())
+	{ // no bone has been removed
+		UE_LOG(LogAnimation, Error, TEXT("Skeleton Modifier - Remove: No bone has been removed."));
+		return false;
+	}
+	
+	// invalidate composer
+	TransformComposer->Invalidate(INDEX_NONE);
+
+	// update index tracker
+	UpdateBoneTracker(InfosBeforeRemoval);
+
+	return true;		
+}
+
+bool FSkeletonModifier::RenameBone(const FName InOldBoneName, const FName InNewBoneName)
+{
+	if (InOldBoneName == NAME_None || InOldBoneName == NAME_None || InNewBoneName == InOldBoneName)
+	{
+		UE_LOG(LogAnimation, Error, TEXT("Skeleton Modifier - Rename: cannot rename %s with %s."), *InOldBoneName.ToString(), *InNewBoneName.ToString());
+		return false;
+	}
+	
+	return RenameBones({InOldBoneName}, {InNewBoneName});
+}
+
+bool FSkeletonModifier::RenameBones(const TArray<FName>& InOldBoneNames, const TArray<FName>& InNewBoneNames)
+{
+	if (!IsReferenceSkeletonValid())
+	{
+		return false;
+	}
+	
+	if (InOldBoneNames.IsEmpty() || InNewBoneNames.Num() != InOldBoneNames.Num())
+	{
+		UE_LOG(LogAnimation, Error, TEXT("Skeleton Modifier - Rename: Discrepancy between old and new names (%d / %d)."), InOldBoneNames.Num(), InNewBoneNames.Num());
+		return false;
+	}
+
+	// update reference skeleton
+	{
+		FReferenceSkeletonModifier Modifier(*ReferenceSkeleton, nullptr);
+
+		const int32 NumBonesToRename = InOldBoneNames.Num();
+		for (int32 Index = 0; Index < NumBonesToRename; ++Index)
+		{
+			const FName& OldName = InOldBoneNames[Index];
+			const FName& NewName = InNewBoneNames[Index];
+			if (OldName != NAME_None && OldName != NewName)
+			{
+				Modifier.Rename(OldName, NewName);
+			}
+		}
+	}
+
+	// update index tracker: no modification on bone indices to track when renaming
+
+	return true;
+}
+
+bool FSkeletonModifier::ParentBone(const FName InBoneName, const FName InParentName)
+{
+	if (InBoneName == NAME_None)
+	{
+		UE_LOG(LogAnimation, Error, TEXT("Skeleton Modifier - Parent: Cannot parent a bone with no name."));
+		return false;
+	}
+	
+	return ParentBones({InBoneName}, {InParentName});
+}
+
+bool FSkeletonModifier::ParentBones(const TArray<FName>& InBoneNames, const TArray<FName>& InParentNames)
+{
+	if (!IsReferenceSkeletonValid())
+	{
+		return false;
+	}
+	
+	if (InBoneNames.IsEmpty())
+	{
+		UE_LOG(LogAnimation, Error, TEXT("Skeleton Modifier - Parent: No bone provided."));
+		return false;
+	}
+
+	// store initial data
+	const TArray<FMeshBoneInfo> InfosBeforeParenting = ReferenceSkeleton->GetRawRefBoneInfo();
+	
+	// update reference skeleton
+	{
+		auto GetParentName = [&](int32 Index)
+		{
+			if (InBoneNames.Num() == InParentNames.Num())
+			{
+				return InParentNames[Index];
+			}
+			return InParentNames.IsEmpty() ? NAME_None : InParentNames[0];  
+		};
+
+		static constexpr bool bAllowMultipleRoots = true;
+		
+		FReferenceSkeletonModifier Modifier(*ReferenceSkeleton, nullptr);
+		for (int32 Index = 0; Index < InBoneNames.Num(); ++Index)
+		{
+			const int32 BoneIndex = ReferenceSkeleton->FindRawBoneIndex(InBoneNames[Index]);
+			if (BoneIndex != INDEX_NONE)
+			{
+				const FName NewParentName = GetParentName(Index);
+
+				// store global transforms
+				const FTransform GlobalTransform = TransformComposer->GetGlobalTransform(BoneIndex);
+
+				const int32 ParentIndex = ReferenceSkeleton->FindRawBoneIndex(NewParentName);
+				const FTransform ParentTransform = TransformComposer->GetGlobalTransform(ParentIndex);
+				
+				// change parent
+				const int32 NewIndex = Modifier.SetParent(InBoneNames[Index], NewParentName, bAllowMultipleRoots);
+				if (NewIndex > INDEX_NONE)
+				{
+					// update local transform
+					Modifier.UpdateRefPoseTransform(NewIndex, GlobalTransform.GetRelativeTransform(ParentTransform));
+				}
+
+				// invalidate composer
+				TransformComposer->Invalidate(INDEX_NONE);
+			}
+		}
+	}
+
+	// update index tracker
+	UpdateBoneTracker(InfosBeforeParenting);
+	
+	return true;
+}
+
+bool FSkeletonModifier::AlignBone(const FName InBoneName, const bool bAlignChildren)
+{
+	if (InBoneName == NAME_None)
+	{
+		UE_LOG(LogAnimation, Error, TEXT("Skeleton Modifier - Align: Cannot align a bone with no name."));
+		return false;
+	}
+	
+	return AlignBones({InBoneName}, bAlignChildren);
+}
+
+bool FSkeletonModifier::AlignBones(const TArray<FName>& InBoneNames, const bool bAlignChildren)
+{
+	if (!IsReferenceSkeletonValid())
+	{
+		return false;
+	}
+	
+	if (InBoneNames.IsEmpty())
+	{
+		UE_LOG(LogAnimation, Error, TEXT("Skeleton Modifier - Align: No bone provided."));
+		return false;
+	}
+
+	auto GetAlignedTransform = [&](const FName& InBoneName) -> TOptional<FTransform>
+	{
+		TOptional<FTransform> OptTransform;
+		
+		const int32 BoneIndex = ReferenceSkeleton->FindRawBoneIndex(InBoneName);
+		const int32 ParentIndex = BoneIndex != INDEX_NONE ? ReferenceSkeleton->GetRawParentIndex(BoneIndex) : INDEX_NONE;
+		if (ParentIndex == INDEX_NONE)
+		{ // we can't align if there's no parent
+			return OptTransform;
+		}
+
+		TArray<int32> Children; ReferenceSkeleton->GetDirectChildBones(BoneIndex, Children);
+		if (Children.Num() != 1)
+		{ // we can't align if there are more than one children
+			return OptTransform;
+		}
+
+		const FTransform& ParentGlobal = GetTransform(ParentIndex, true);
+		const FTransform& BoneGlobal = GetTransform(BoneIndex, true);
+		const FTransform& ChildGlobal = GetTransform(Children[0], true);
+
+		// X/Y axis
+		const FVector XAxis = (ChildGlobal.GetLocation() - BoneGlobal.GetLocation()).GetSafeNormal();
+		const FVector YAxis = (BoneGlobal.GetLocation() - ParentGlobal.GetLocation()).GetSafeNormal();
+
+		if (XAxis.IsNearlyZero() || YAxis.IsNearlyZero())
+		{
+			return OptTransform;
+		}
+
+		// Z axis
+		const bool bComputePlane = (FMath::Abs(FVector::DotProduct(XAxis, YAxis)) - 1.0f) < KINDA_SMALL_NUMBER;
+		const FVector ZAxis = bComputePlane ? FVector::CrossProduct(XAxis, YAxis) : BoneGlobal.TransformVector(FVector::ZAxisVector);
+	
+		FQuat Rotation = FRotationMatrix::MakeFromXZ(XAxis, ZAxis).ToQuat();
+		Rotation.Normalize();
+
+		const FTransform GlobalTransform(Rotation, BoneGlobal.GetLocation());
+		const FTransform LocalTransform = GlobalTransform.GetRelativeTransform(ParentGlobal);
+
+		return LocalTransform;
+	};
+	
+	TArray<FName> BonesToAlign; BonesToAlign.Reserve(InBoneNames.Num());
+	TArray<FTransform> Transforms; Transforms.Reserve(InBoneNames.Num());
+	for (int32 Index = 0; Index < InBoneNames.Num(); ++Index)	
+	{
+		TOptional<FTransform> OptTransform = GetAlignedTransform(InBoneNames[Index]);
+		if (OptTransform)
+		{
+			BonesToAlign.Add(InBoneNames[Index]);
+			Transforms.Add(*OptTransform);
+		}
+	}
+
+	if (BonesToAlign.IsEmpty())
+	{
+		UE_LOG(LogAnimation, Error, TEXT("Skeleton Modifier - Align: No bone to align."));
+		return false;
+	}
+
+	static constexpr bool bMoveChildren = false;
+	return MoveBones(BonesToAlign, Transforms, bMoveChildren);
+}
+
+void FSkeletonModifier::UpdateBoneTracker(const TArray<FMeshBoneInfo>& InOtherInfos)
+{
+	for (int32 Index = 0; Index < BoneIndexTracker.Num(); ++Index)
+	{
+		const int32 IndexToTrack = BoneIndexTracker[Index];
+		if (IndexToTrack > INDEX_NONE)
+		{
+			check(InOtherInfos.IsValidIndex(IndexToTrack));
+			const int32 NewIndex = ReferenceSkeleton->FindRawBoneIndex(InOtherInfos[IndexToTrack].Name);
+			BoneIndexTracker[Index] = NewIndex;
+		}
+	}
+}
+
+FName FSkeletonModifier::GetUniqueName(const FName InBoneName) const
+{
+	if (!ReferenceSkeleton)
+	{
+		return NAME_None;
+	}
+	
+	auto IsNameAvailable = [&](const FName Name)
+	{
+		const int32 Index = ReferenceSkeleton->FindRawBoneIndex(Name);
+		if (Index != INDEX_NONE)
+		{
+			return false;
+		}
+		return true;
+	};
+
+	FName OutName = InBoneName;
+	int32 CurrentIndex = 0;
+
+	while (!IsNameAvailable(OutName))
+	{
+		OutName = FName(*FString::Printf(TEXT("%s%d"), *InBoneName.ToString(), CurrentIndex++));
+	}
+
+	return OutName;
+}
+
+const FReferenceSkeleton& FSkeletonModifier::GetReferenceSkeleton() const
+{
+	static const FReferenceSkeleton Dummy;
+	return ReferenceSkeleton ? *ReferenceSkeleton.Get() : Dummy;
+}
+
+const TArray<int32>& FSkeletonModifier::GetBoneIndexTracker() const
+{
+	return BoneIndexTracker;	
+}
+
+const FTransform& FSkeletonModifier::GetTransform(const int32 InBoneIndex, const bool bGlobal) const
+{
+	if (!ReferenceSkeleton.IsValid())
+	{
+		return FTransform::Identity;
+	}
+	
+	if (bGlobal)
+	{
+		return TransformComposer.IsValid() ? TransformComposer->GetGlobalTransform(InBoneIndex) : FTransform::Identity;  
+	}
+
+	const TArray<FTransform>& LocalTransforms = ReferenceSkeleton->GetRawRefBonePose();
+	return LocalTransforms.IsValidIndex(InBoneIndex) ? LocalTransforms[InBoneIndex] : FTransform::Identity;
+}
+
+
diff --git a/Engine/Source/Developer/SkeletalMeshUtilitiesCommon/Public/SkeletonModifier.h b/Engine/Source/Developer/SkeletalMeshUtilitiesCommon/Public/SkeletonModifier.h
new file mode 100644
--- /dev/null
+++ b/Engine/Source/Developer/SkeletalMeshUtilitiesCommon/Public/SkeletonModifier.h
@@ -0,0 +1,199 @@
+﻿// Copyright Epic Games, Inc. All Rights Reserved.
+
+#pragma once
+
+#include "Math/MathFwd.h"
+#include "Containers/Array.h"
+#include "Templates/UniquePtr.h"
+#include "MeshDescription.h"
+#include "ReferenceSkeleton.h"
+
+#include "SkeletonModifier.generated.h"
+
+class USkeletalMesh;
+class FName;
+struct FMeshDescription;
+struct FReferenceSkeleton;
+struct FTransformComposer;
+
+USTRUCT(BlueprintType)
+struct SKELETALMESHUTILITIESCOMMON_API FMirrorOptions
+{
+	GENERATED_USTRUCT_BODY()
+
+	UPROPERTY(BlueprintReadWrite, EditAnywhere, Category = "Mirror")
+	TEnumAsByte<EAxis::Type> MirrorAxis = EAxis::X;
+
+	UPROPERTY(BlueprintReadWrite, EditAnywhere, Category = "Mirror")
+	bool bMirrorRotation = true;
+
+	UPROPERTY(BlueprintReadWrite, EditAnywhere, Category = "Mirror")
+	FString LeftString = TEXT("_l");
+
+	UPROPERTY(BlueprintReadWrite, EditAnywhere, Category = "Mirror")
+	FString RightString = TEXT("_r");
+
+	UPROPERTY(BlueprintReadWrite, EditAnywhere, Category = "Mirror")
+	bool bMirrorChildren = true;
+
+	FTransform MirrorTransform(const FTransform& InGlobalTransform) const;
+	FVector MirrorVector(const FVector& InVector) const;
+};
+
+/**
+ * FSkeletalMeshSkeletonModifier
+ */
+
+class SKELETALMESHUTILITIESCOMMON_API FSkeletonModifier
+{
+public:
+
+	bool Init(USkeletalMesh* InSkeletalMesh);
+	
+	void ExternalUpdate(const FReferenceSkeleton& InRefSkeleton, const TArray<int32>& InIndexTracker);
+
+	/** Creates a new bone in the skeleton hierarchy at desired transform
+	 *  @param InBoneName The new bone's name.
+	 *  @param InParentName The new bone parent's name. 
+	 *  @param InTransform The default local transform in the parent space.
+	 *  @return \c true if the operation succeeded, false otherwise. 
+	 */
+	bool AddBone(const FName InBoneName, const FName InParentName, const FTransform& InTransform);
+	bool AddBones(const TArray<FName>& InBonesName, const TArray<FName>& InParentsName, const TArray<FTransform>& InTransforms);
+
+	/** Mirror bones
+	 *  @param InBoneName The new bone's name.
+	 *  @param InOptions The mirroring options
+	 *  @return \c true if the operation succeeded, false otherwise. 
+	 */
+	bool MirrorBone(const FName InBoneName, const FMirrorOptions& InOptions = FMirrorOptions());
+	bool MirrorBones(const TArray<FName>& InBonesName, const FMirrorOptions& InOptions = FMirrorOptions());
+	bool MirrorBonesDep(const TArray<FName>& InBonesName);
+
+	/** Moves a bone at a new desired local transform
+	 *  @param InBoneToMove The new bone's name that needs to be moved.
+	 *  @param InNewTransform The new local transform in the bone's parent space.
+	 *  @param bMoveChildren Propagate new transform to children
+	 *  @return \c true if the operation succeeded, false otherwise. 
+	 */
+	bool MoveBone(const FName InBoneToMove, const FTransform& InNewTransform, const bool bMoveChildren);
+	bool MoveBones(const TArray<FName>& InBonesToMove, const TArray<FTransform>& InNewTransforms, const bool bMoveChildren);
+
+	/** Remove a bone in the skeleton hierarchy
+	 *  @param InBoneName The new bone's name.
+	 *  @param bRemoveChildren Remove children recursively.
+	 *  @return \c true if the operation succeeded, false otherwise. 
+	 */
+	bool RemoveBone(const FName InBoneName, const bool bRemoveChildren);
+	bool RemoveBones(const TArray<FName>& InBoneNames, const bool bRemoveChildren);
+
+	/** Rename bones
+	 *  @param InOldBoneName The current bone's name.
+	 *  @param InNewBoneName The new bone's name.
+	 *  @return \c true if the operation succeeded, false otherwise. 
+	 */
+	bool RenameBone(const FName InOldBoneName, const FName InNewBoneName);
+	bool RenameBones(const TArray<FName>& InOldBoneNames, const TArray<FName>& InNewBoneNames);
+	
+	/** Parent bones
+	 *  @param InBoneName The current bone's name.
+	 *  @param InParentName The new parent's name (Name_NONE to unparent).
+	 *  @return \c true if the operation succeeded, false otherwise. 
+	 */
+	bool ParentBone(const FName InBoneName, const FName InParentName);
+	bool ParentBones(const TArray<FName>& InBoneNames, const TArray<FName>& InParentNames);
+
+	/** Align bones
+	 *  @param InBoneName The current bone's name.
+	 *  @param bAlignChildren
+	 *  @return \c true if the operation succeeded, false otherwise. 
+	 */
+	bool AlignBone(const FName InBoneName, const bool bAlignChildren);
+	bool AlignBones(const TArray<FName>& InBoneNames, const bool bAlignChildren);
+	
+	/**
+	 * Actually applies the skeleton modifications to the skeletal mesh.
+	 * @return true if commit succeeded.
+	 */
+	bool CommitSkeletonToSkeletalMesh();
+
+	FName GetUniqueName(const FName InBoneName) const;
+	
+	const FReferenceSkeleton& GetReferenceSkeleton() const;
+	const TArray<int32>& GetBoneIndexTracker() const;
+	const FTransform& GetTransform(const int32 InBoneIndex, const bool bGlobal = false) const;
+	
+private:
+
+	bool IsReferenceSkeletonValid(const bool bLog = true) const;
+
+	void UpdateBoneTracker(const TArray<FMeshBoneInfo>& InOtherInfos);
+	
+	// mirroring functions
+	void GetBonesToMirror(const TArray<FName>& InBonesName, const FMirrorOptions& InOptions, TArray<int32>& OutBonesToMirror) const;
+	void GetMirroredNames(const TArray<int32>& InBonesToMirror, const FMirrorOptions& InOptions, TArray<FName>& OutBonesName) const;
+	void GetMirroredBones(const TArray<int32>& InBonesToMirror, const TArray<FName>& InMirroredNames, TArray<int32>& OutMirroredBones);
+	void GetMirroredTransforms(const TArray<int32>& InBonesToMirror, const FMirrorOptions& InOptions, TArray<FTransform>& OutMirroredTransforms) const;
+	
+	TObjectPtr<USkeletalMesh> SkeletalMesh = nullptr;
+	TUniquePtr<FMeshDescription> MeshDescription;
+	TUniquePtr<FReferenceSkeleton> ReferenceSkeleton;
+	TUniquePtr<FTransformComposer> TransformComposer;
+	TArray<int32> BoneIndexTracker;
+};
+
+struct FTransformComposer
+{
+	FTransformComposer(const FReferenceSkeleton& InRefSkeleton)
+		: RefSkeleton(InRefSkeleton)
+		, Transforms(InRefSkeleton.GetRawRefBonePose())
+	{
+		TransformCached.Init(false, Transforms.Num());
+	}
+
+	const FTransform& GetGlobalTransform(const uint32 Index)
+	{
+		if (!Transforms.IsValidIndex(Index))
+		{
+			return FTransform::Identity;
+		}
+
+		if (TransformCached[Index])
+		{
+			return Transforms[Index];
+		}
+
+		const int32 ParentIndex = RefSkeleton.GetRawParentIndex(Index);
+		if (ParentIndex > INDEX_NONE)
+		{
+			Transforms[Index] *= GetGlobalTransform(ParentIndex);
+		}
+
+		TransformCached[Index] = true;
+		return Transforms[Index];
+	}
+
+	void Invalidate(const uint32 Index = INDEX_NONE)
+	{
+		if (!TransformCached.IsValidIndex(Index))
+		{
+			Transforms = RefSkeleton.GetRawRefBonePose();
+			TransformCached.Init(false, Transforms.Num());
+			return;
+		}
+		
+		Transforms[Index] = RefSkeleton.GetRawRefBonePose()[Index];
+		TransformCached[Index] = false;
+
+		TArray<int32> Children; RefSkeleton.GetDirectChildBones(Index, Children);
+		for (const int32 ChildIndex: Children)
+		{
+			Invalidate(ChildIndex);
+		}
+	}
+
+	const FReferenceSkeleton& RefSkeleton;
+	TArray<FTransform> Transforms;
+	TBitArray<> TransformCached;
+};
+
diff --git a/Engine/Source/Runtime/Engine/Private/Animation/ReferenceSkeleton.cpp b/Engine/Source/Runtime/Engine/Private/Animation/ReferenceSkeleton.cpp
--- a/Engine/Source/Runtime/Engine/Private/Animation/ReferenceSkeleton.cpp
+++ b/Engine/Source/Runtime/Engine/Private/Animation/ReferenceSkeleton.cpp
@@ -5,7 +5,10 @@
 #include "EngineLogs.h"
 #include "Engine/SkeletalMesh.h"
 
-FReferenceSkeletonModifier::FReferenceSkeletonModifier(USkeleton* InSkeleton) : RefSkeleton(InSkeleton->ReferenceSkeleton), Skeleton(InSkeleton) {}
+FReferenceSkeletonModifier::FReferenceSkeletonModifier(USkeleton* InSkeleton)
+	: RefSkeleton(InSkeleton->ReferenceSkeleton),
+	Skeleton(InSkeleton)
+{}
 
 FReferenceSkeletonModifier::~FReferenceSkeletonModifier()
 {
@@ -17,11 +20,28 @@ void FReferenceSkeletonModifier::UpdateRefPoseTransform(const int32 BoneIndex, c
 	RefSkeleton.UpdateRefPoseTransform(BoneIndex, BonePose);
 }
 
-void FReferenceSkeletonModifier::Add(const FMeshBoneInfo& BoneInfo, const FTransform& BonePose)
+void FReferenceSkeletonModifier::Add(const FMeshBoneInfo& BoneInfo, const FTransform& BonePose, const bool bAllowMultipleRoots)
 {
+	TGuardValue OperationGuard(RefSkeleton.bOnlyOneRootAllowed, !bAllowMultipleRoots);
 	RefSkeleton.Add(BoneInfo, BonePose);
 }
 
+void FReferenceSkeletonModifier::Remove(const FName& BoneName, const bool bRemoveChildren)
+{
+	RefSkeleton.Remove(BoneName, bRemoveChildren);
+}
+
+void FReferenceSkeletonModifier::Rename(const FName& InOldName, const FName& InNewName)
+{
+	RefSkeleton.Rename(InOldName, InNewName);
+}
+
+int32 FReferenceSkeletonModifier::SetParent(const FName& InBoneName, const FName& InParentName, const bool bAllowMultipleRoots)
+{
+	TGuardValue OperationGuard(RefSkeleton.bOnlyOneRootAllowed, !bAllowMultipleRoots);
+	return RefSkeleton.SetParent(InBoneName, InParentName);
+}
+
 int32 FReferenceSkeletonModifier::FindBoneIndex(const FName& BoneName) const
 {
 	return RefSkeleton.FindRawBoneIndex(BoneName);
@@ -75,6 +95,246 @@ FTransform GetComponentSpaceTransform(TArray<uint8>& ComponentSpaceFlags, TArray
 	return This;
 }
 
+void FReferenceSkeleton::Remove(const FName InBoneName, const bool bRemoveChildren)
+{
+	const int32 RawBoneIndex = FindRawBoneIndex(InBoneName);
+	if (RawBoneIndex == INDEX_NONE)
+	{
+		return;
+	}
+
+	const FMeshBoneInfo& BoneInfo = RawRefBoneInfo[RawBoneIndex];
+	const int32 RawParentIndex = BoneInfo.ParentIndex;
+	
+	// Make sure our arrays are in sync.
+	checkSlow((RawRefBoneInfo.Num() == RawRefBonePose.Num()) && (RawRefBoneInfo.Num() == RawNameToIndexMap.Num()));
+
+	// store children indices and sort them from greatest to lowest
+	TArray<int32> Children;
+	GetDirectChildBones(RawBoneIndex, Children);
+	Children.Sort([](const int32 Index0, const int32 Index1) {return Index0 > Index1;} );
+
+	// reindex function
+	auto ReIndexBones = [this, RawBoneIndex, RawParentIndex](const bool bUpdateParent)
+	{
+		for (int32 NextIndex=RawBoneIndex+1; NextIndex < GetRawBoneNum(); NextIndex++)
+		{
+			FMeshBoneInfo& Bone = RawRefBoneInfo[NextIndex];
+        
+			// update parent
+			if (Bone.ParentIndex > RawBoneIndex)
+			{
+				Bone.ParentIndex -= 1;
+			}
+			else if (bUpdateParent && (Bone.ParentIndex == RawBoneIndex))
+			{
+				Bone.ParentIndex = RawParentIndex;
+			}
+        
+			// update cached index
+			RawNameToIndexMap[Bone.Name] -= 1;
+		}
+	};
+
+	if (bRemoveChildren)
+	{
+		// 1 - treat children first
+		for (const int32 ChildIndex: Children)
+		{
+			Remove(RawRefBoneInfo[ChildIndex].Name, bRemoveChildren);
+		}
+
+		// 2 - reindex next bones
+		ReIndexBones(false);
+		
+		// 3 - remove useless raw data
+		RawRefBonePose.RemoveAt(RawBoneIndex, 1);
+		RawRefBoneInfo.RemoveAt(RawBoneIndex, 1);
+		RawNameToIndexMap.Remove(InBoneName);
+		
+		return;
+	}
+
+	// 1 - store transforms
+
+	// store parent's global transform
+	FTransform ParentGlobal = FTransform::Identity;
+	int32 ParentIndex = RawParentIndex;
+	while (ParentIndex > INDEX_NONE)
+	{
+		ParentGlobal = ParentGlobal * RawRefBonePose[ParentIndex];
+		ParentIndex--;
+	}
+
+	// store bone's global transform
+	const FTransform BoneGlobal = RawRefBonePose[RawBoneIndex] * ParentGlobal;
+
+	// 2 - switch children transforms to new parent space
+	for (const int32 ChildIndex: Children)
+	{
+		const FTransform ChildrenGlobal = RawRefBonePose[ChildIndex] * BoneGlobal;
+		RawRefBonePose[ChildIndex] = ChildrenGlobal.GetRelativeTransform(ParentGlobal);
+	}
+
+	// 3 - reindex next bones
+	ReIndexBones(true);
+
+	// 4 - remove useless raw data
+	RawRefBonePose.RemoveAt(RawBoneIndex, 1);
+	RawRefBoneInfo.RemoveAt(RawBoneIndex, 1);
+	RawNameToIndexMap.Remove(InBoneName);
+}
+
+void FReferenceSkeleton::Rename(const FName InBoneName, const FName InNewName)
+{
+	const int32 RawBoneIndex = FindRawBoneIndex(InBoneName);
+	const int32 RawNewBoneIndex = FindRawBoneIndex(InNewName);
+	if (RawBoneIndex == INDEX_NONE || RawNewBoneIndex != INDEX_NONE)
+	{
+		return;
+	}
+
+	FMeshBoneInfo& BoneInfo = RawRefBoneInfo[RawBoneIndex];
+	BoneInfo.Name = InNewName;
+	RawNameToIndexMap.Remove(InBoneName);
+	RawNameToIndexMap.Add(InNewName, RawBoneIndex);
+}
+
+int32 FReferenceSkeleton::SetParent(const FName InBoneName, const FName InParentName)
+{
+	if (InBoneName == InParentName)
+	{
+		return INDEX_NONE;
+	}
+
+	const int32 BoneIndex = FindRawBoneIndex(InBoneName);
+	if (BoneIndex < 1)
+	{ // we do not allow to change the root's parent 
+		return BoneIndex;
+	}
+	
+	const int32 NewParentIndex = FindRawBoneIndex(InParentName);
+	if (NewParentIndex == INDEX_NONE && InParentName != NAME_None)
+	{
+		return INDEX_NONE;
+	}
+
+	auto GetParentIndex = [&](const int32 InBoneIndex) -> int32
+	{
+		return RawRefBoneInfo.IsValidIndex(InBoneIndex) ? RawRefBoneInfo[InBoneIndex].ParentIndex : INDEX_NONE;
+	};
+
+	const int32 NumBones = GetRawBoneNum();
+	auto GetEndOfBranch = [&](const int32 InBoneIndex) -> int32
+	{
+		const int32 StartBranch = InBoneIndex; 
+
+		int32 EndOfBranch = InBoneIndex+1;
+		int32 ParentIndex = GetParentIndex(EndOfBranch);
+
+		// if next child bone's parent is less than or equal to StartParentIndex,
+		// we are leaving the branch so no need to go further
+		while (ParentIndex >= StartBranch && EndOfBranch < NumBones)
+		{
+			EndOfBranch++;
+			ParentIndex = GetParentIndex(EndOfBranch);
+		}
+		return EndOfBranch;
+	};
+
+	// bone branch to move 
+	const int32 EndOfBoneBranch = GetEndOfBranch(BoneIndex);
+
+	// parent branch to move
+	const int32 EndOfNewParentBranch = NewParentIndex > INDEX_NONE ? GetEndOfBranch(NewParentIndex) : EndOfBoneBranch+1;
+
+	// reorder bones
+	TArray<FName> Names;
+	auto AddBones = [&](int32 Begin, int32 End)
+	{
+		for (int32 Index = Begin; Index < End; Index++)
+		{
+			Names.Add(RawRefBoneInfo[Index].Name);
+		}		
+	};
+
+	if (NewParentIndex == INDEX_NONE)
+	{
+		// add all bones before BoneIndex
+		AddBones(0, BoneIndex);
+		// add bones after branch
+		AddBones(EndOfBoneBranch, NumBones);
+		// add bone branch at the end
+		AddBones(BoneIndex, EndOfBoneBranch);
+	}
+	else if (NewParentIndex > BoneIndex)
+	{
+		// add all bones before BoneIndex
+		AddBones(0, BoneIndex);
+		// add parent branch up
+		AddBones(NewParentIndex, EndOfNewParentBranch);
+		// add bone branch
+		AddBones(BoneIndex, EndOfBoneBranch);
+		// add bones between bone branch and parent 
+		AddBones(EndOfBoneBranch, NewParentIndex);
+		// add the end
+		AddBones(EndOfNewParentBranch, NumBones);
+	}
+	else
+	{
+		// add all bones before the end of the new parent branch
+		AddBones(0, EndOfNewParentBranch);
+		// add bone branch
+		AddBones(BoneIndex, EndOfBoneBranch);
+		// add remaining bones
+		AddBones(EndOfNewParentBranch, BoneIndex);
+		AddBones(EndOfBoneBranch, NumBones);
+	}
+
+	check(Names.Num() == NumBones);
+	
+	TMap<int32, int32> OldToNewIndexes; OldToNewIndexes.Reserve(Names.Num());
+	TMap<FName, int32> NewNameToIndexMap; NewNameToIndexMap.Reserve(Names.Num());
+	TArray<FMeshBoneInfo> NewRawRefBoneInfo; NewRawRefBoneInfo.Reserve(Names.Num());
+	TArray<FTransform> NewRawRefBonePose; NewRawRefBonePose.Reserve(Names.Num());
+
+	// store new data
+	for (int32 Index = 0; Index < Names.Num(); ++Index)
+	{
+		const int32 OldIndex = FindRawBoneIndex(Names[Index]);
+		OldToNewIndexes.Add(OldIndex, Index);
+		
+		NewNameToIndexMap.Add(Names[Index], Index);
+		NewRawRefBoneInfo.Add(RawRefBoneInfo[OldIndex]);
+		NewRawRefBonePose.Add(RawRefBonePose[OldIndex]);
+	}
+
+	// new bone index
+	const int32 NewBoneIndex = *NewNameToIndexMap.Find(InBoneName);
+	
+	// update parents
+	for (int32 Index = 0; Index < Names.Num(); ++Index)
+	{
+		FMeshBoneInfo& NewBoneInfo = NewRawRefBoneInfo[Index];
+		if (Index == NewBoneIndex)
+		{
+			NewBoneInfo.ParentIndex = NewParentIndex > INDEX_NONE ? *NewNameToIndexMap.Find(InParentName) : INDEX_NONE;
+		}
+		else if (NewBoneInfo.ParentIndex > INDEX_NONE)
+		{
+			NewBoneInfo.ParentIndex = *OldToNewIndexes.Find(NewBoneInfo.ParentIndex);
+		}
+	}
+	
+	// swap data
+	RawRefBonePose = MoveTemp(NewRawRefBonePose);
+	RawRefBoneInfo = MoveTemp(NewRawRefBoneInfo);
+	RawNameToIndexMap = MoveTemp(NewNameToIndexMap);
+	
+	// return new bone index
+	return NewBoneIndex;
+}
+
 int32 FReferenceSkeleton::GetRawSourceBoneIndex(const USkeleton* Skeleton, const FName& SourceBoneName) const
 {
 	for (const FVirtualBone& VB : Skeleton->GetVirtualBones())
diff --git a/Engine/Source/Runtime/Engine/Private/Rendering/SkeletalMeshLODImporterData.cpp b/Engine/Source/Runtime/Engine/Private/Rendering/SkeletalMeshLODImporterData.cpp
--- a/Engine/Source/Runtime/Engine/Private/Rendering/SkeletalMeshLODImporterData.cpp
+++ b/Engine/Source/Runtime/Engine/Private/Rendering/SkeletalMeshLODImporterData.cpp
@@ -17,6 +17,7 @@
 #include "Serialization/MemoryReader.h"
 #include "SkeletalMeshAttributes.h"
 #include "SkeletalMeshOperations.h"
+#include "MeshElementContainer.h"
 
 
 DEFINE_LOG_CATEGORY_STATIC(LogSkeletalMeshLODImporterData, Log, All);
@@ -2008,6 +2009,39 @@ FSkeletalMeshImportData FSkeletalMeshImportData::CreateFromMeshDescription(const
 		}
 	}
 
+	// Update Bones data
+	const FSkeletalMeshAttributes::FBoneArray& Bones = Attributes.Bones();
+	SkelMeshImportData.RefBonesBinary.Reserve(Bones.Num());
+	
+	FSkeletalMeshAttributes::FBoneNameAttributesConstRef BoneNames = Attributes.GetBoneNames();
+	FSkeletalMeshAttributes::FBoneParentIndexAttributesConstRef BoneParentIndices = Attributes.GetBoneParentIndices();
+	FSkeletalMeshAttributes::FBonePoseAttributesConstRef BonePoses = Attributes.GetBonePoses();
+
+	auto GetNumChildren = [&](const int32 ParentIndex)
+	{
+		int32 NumChildren = 0;
+		for (const FBoneID& BoneId: Bones.GetElementIDs())
+		{
+			if (BoneParentIndices.Get(BoneId) == ParentIndex)
+			{
+				NumChildren++;
+			}
+		}
+		return NumChildren;
+	};
+
+	int32 Index = 0;
+	for (const FBoneID& BoneId: Bones.GetElementIDs())
+	{
+		SkeletalMeshImportData::FBone NewBone;
+			NewBone.Name = BoneNames.Get(BoneId).ToString();
+			NewBone.ParentIndex = BoneParentIndices.Get(BoneId);
+			NewBone.BonePos.Transform = FTransform3f(BonePoses.Get(BoneId));
+			NewBone.NumChildren = GetNumChildren(Index);
+		SkelMeshImportData.RefBonesBinary.Emplace( MoveTemp(NewBone) );
+		Index++;
+	}
+
 	SkelMeshImportData.CleanUpUnusedMaterials();
 	SkelMeshImportData.SplitVerticesBySmoothingGroups();
 
diff --git a/Engine/Source/Runtime/Engine/Public/ReferenceSkeleton.h b/Engine/Source/Runtime/Engine/Public/ReferenceSkeleton.h
--- a/Engine/Source/Runtime/Engine/Public/ReferenceSkeleton.h
+++ b/Engine/Source/Runtime/Engine/Public/ReferenceSkeleton.h
@@ -83,8 +83,17 @@ public:
 	void UpdateRefPoseTransform(const int32 BoneIndex, const FTransform& BonePose);
 
 	// Add a new bone. BoneName must not already exist! ParentIndex must be valid.
-	void Add(const FMeshBoneInfo& BoneInfo, const FTransform& BonePose);
+	void Add(const FMeshBoneInfo& BoneInfo, const FTransform& BonePose, const bool bAllowMultipleRoots = false);
 
+	// Remove a bone. BoneName must be valid.
+	void Remove(const FName& BoneName, const bool bRemoveChildren);
+
+	// Rename a bone. InOldName must be valid and InNewName not already a bone name.
+	void Rename(const FName& InOldName, const FName& InNewName);
+
+	// Change bone's parent. InBoneName must be valid and InParentName can be Name_NONE to unparent.
+	int32 SetParent(const FName& InBoneName, const FName& InParentName, const bool bAllowMultipleRoots = false);
+	
 	/** Find Bone Index from BoneName. Precache as much as possible in speed critical sections! */
 	int32 FindBoneIndex(const FName& BoneName) const;
 
@@ -204,6 +213,15 @@ private:
 			|| ((BoneIndex > 0) && RawRefBoneInfo.IsValidIndex(BoneInfo.ParentIndex) && (BoneInfo.ParentIndex < BoneIndex))));
 	}
 
+	/** Remove InBoneName and its children if bRemoveChildren == true. */
+	void Remove(const FName InBoneName, const bool bRemoveChildren);
+
+	/** Rename InBoneName with InNewName. */
+	void Rename(const FName InBoneName, const FName InNewName);
+
+	/** Set InParentName as InBoneName's parent. */
+	int32 SetParent(const FName InBoneName, const FName InParentName);
+
 	// Help us translate a virtual bone source into a raw bone source (for evaluating virtual bone transform)
 	int32 GetRawSourceBoneIndex(const USkeleton* Skeleton, const FName& SourceBoneName) const;
 
@@ -250,6 +268,8 @@ public:
 		return RawRefBonePose;
 	}
 
+	const TMap<FName, int32>& GetRawNameToIndexMap() const { return RawNameToIndexMap; }
+	
 	void Empty(int32 Size=0)
 	{
 		RawRefBoneInfo.Empty(Size);
