Rematch material on re-import for both the static mesh and the skeletal mesh.
Add a new factory step call ImportAssetObjectFinalize_GameThread so we can call some UI on the game thread when importing an asset.

#rnx

diff --git a/Engine/Plugins/Experimental/Enterprise/DatasmithInterchange/Source/DatasmithInterchange/Private/InterchangeDatasmithSceneFactory.cpp b/Engine/Plugins/Experimental/Enterprise/DatasmithInterchange/Source/DatasmithInterchange/Private/InterchangeDatasmithSceneFactory.cpp
--- a/Engine/Plugins/Experimental/Enterprise/DatasmithInterchange/Source/DatasmithInterchange/Private/InterchangeDatasmithSceneFactory.cpp
+++ b/Engine/Plugins/Experimental/Enterprise/DatasmithInterchange/Source/DatasmithInterchange/Private/InterchangeDatasmithSceneFactory.cpp
@@ -15,7 +15,7 @@ UClass* UInterchangeDatasmithSceneFactory::GetFactoryClass() const
 }
 
 
-UObject* UInterchangeDatasmithSceneFactory::ImportAssetObject_GameThread(const FImportAssetObjectParams& Arguments)
+UObject* UInterchangeDatasmithSceneFactory::BeginImportAssetObject_GameThread(const FImportAssetObjectParams& Arguments)
 {
 	UDatasmithScene* DatasmithScene = nullptr;
 
diff --git a/Engine/Plugins/Experimental/Enterprise/DatasmithInterchange/Source/DatasmithInterchange/Public/InterchangeDatasmithSceneFactory.h b/Engine/Plugins/Experimental/Enterprise/DatasmithInterchange/Source/DatasmithInterchange/Public/InterchangeDatasmithSceneFactory.h
--- a/Engine/Plugins/Experimental/Enterprise/DatasmithInterchange/Source/DatasmithInterchange/Public/InterchangeDatasmithSceneFactory.h
+++ b/Engine/Plugins/Experimental/Enterprise/DatasmithInterchange/Source/DatasmithInterchange/Public/InterchangeDatasmithSceneFactory.h
@@ -21,7 +21,7 @@ public:
 	// Interchange factory base interface begin
 
 	virtual UClass* GetFactoryClass() const override;
-	virtual UObject* ImportAssetObject_GameThread(const FImportAssetObjectParams& Arguments) override;
+	virtual UObject* BeginImportAssetObject_GameThread(const FImportAssetObjectParams& Arguments) override;
 	virtual UObject* ImportAssetObject_Async(const FImportAssetObjectParams& Arguments) override;
 	virtual void SetupObject_GameThread(const FSetupObjectParams& Arguments) override;
 	// Interchange factory base interface end
diff --git a/Engine/Plugins/Interchange/Runtime/Source/Import/InterchangeImport.Build.cs b/Engine/Plugins/Interchange/Runtime/Source/Import/InterchangeImport.Build.cs
--- a/Engine/Plugins/Interchange/Runtime/Source/Import/InterchangeImport.Build.cs
+++ b/Engine/Plugins/Interchange/Runtime/Source/Import/InterchangeImport.Build.cs
@@ -52,6 +52,7 @@ namespace UnrealBuildTool.Rules
 					{
 						"BSPUtils",
 						"MaterialEditor",
+						"SkeletalMeshUtilitiesCommon",
 						"UnrealEd",
 						"VariantManager",
 					}
diff --git a/Engine/Plugins/Interchange/Runtime/Source/Import/Private/Animation/InterchangeAnimSequenceFactory.cpp b/Engine/Plugins/Interchange/Runtime/Source/Import/Private/Animation/InterchangeAnimSequenceFactory.cpp
--- a/Engine/Plugins/Interchange/Runtime/Source/Import/Private/Animation/InterchangeAnimSequenceFactory.cpp
+++ b/Engine/Plugins/Interchange/Runtime/Source/Import/Private/Animation/InterchangeAnimSequenceFactory.cpp
@@ -767,7 +767,7 @@ UClass* UInterchangeAnimSequenceFactory::GetFactoryClass() const
 	return UAnimSequence::StaticClass();
 }
 
-UObject* UInterchangeAnimSequenceFactory::ImportAssetObject_GameThread(const FImportAssetObjectParams& Arguments)
+UObject* UInterchangeAnimSequenceFactory::BeginImportAssetObject_GameThread(const FImportAssetObjectParams& Arguments)
 {
 #if !WITH_EDITOR || !WITH_EDITORONLY_DATA
 
diff --git a/Engine/Plugins/Interchange/Runtime/Source/Import/Private/Animation/InterchangeAnimationTrackSetFactory.cpp b/Engine/Plugins/Interchange/Runtime/Source/Import/Private/Animation/InterchangeAnimationTrackSetFactory.cpp
--- a/Engine/Plugins/Interchange/Runtime/Source/Import/Private/Animation/InterchangeAnimationTrackSetFactory.cpp
+++ b/Engine/Plugins/Interchange/Runtime/Source/Import/Private/Animation/InterchangeAnimationTrackSetFactory.cpp
@@ -652,7 +652,7 @@ UClass* UInterchangeAnimationTrackSetFactory::GetFactoryClass() const
 	return ULevelSequence::StaticClass();
 }
 
-UObject* UInterchangeAnimationTrackSetFactory::ImportAssetObject_GameThread(const FImportAssetObjectParams& Arguments)
+UObject* UInterchangeAnimationTrackSetFactory::BeginImportAssetObject_GameThread(const FImportAssetObjectParams& Arguments)
 {
 #if !WITH_EDITOR || !WITH_EDITORONLY_DATA
 
diff --git a/Engine/Plugins/Interchange/Runtime/Source/Import/Private/Material/InterchangeMaterialFactory.cpp b/Engine/Plugins/Interchange/Runtime/Source/Import/Private/Material/InterchangeMaterialFactory.cpp
--- a/Engine/Plugins/Interchange/Runtime/Source/Import/Private/Material/InterchangeMaterialFactory.cpp
+++ b/Engine/Plugins/Interchange/Runtime/Source/Import/Private/Material/InterchangeMaterialFactory.cpp
@@ -351,7 +351,7 @@ UClass* UInterchangeMaterialFactory::GetFactoryClass() const
 	return UMaterialInterface::StaticClass();
 }
 
-UObject* UInterchangeMaterialFactory::ImportAssetObject_GameThread(const FImportAssetObjectParams& Arguments)
+UObject* UInterchangeMaterialFactory::BeginImportAssetObject_GameThread(const FImportAssetObjectParams& Arguments)
 {
 	UObject* Material = nullptr;
 
@@ -1179,7 +1179,7 @@ UClass* UInterchangeMaterialFunctionFactory::GetFactoryClass() const
 	return UMaterialFunctionInterface::StaticClass();
 }
 
-UObject* UInterchangeMaterialFunctionFactory::ImportAssetObject_GameThread(const FImportAssetObjectParams& Arguments)
+UObject* UInterchangeMaterialFunctionFactory::BeginImportAssetObject_GameThread(const FImportAssetObjectParams& Arguments)
 {
 	UObject* Material = nullptr;
 
diff --git a/Engine/Plugins/Interchange/Runtime/Source/Import/Private/Mesh/InterchangePhysicsAssetFactory.cpp b/Engine/Plugins/Interchange/Runtime/Source/Import/Private/Mesh/InterchangePhysicsAssetFactory.cpp
--- a/Engine/Plugins/Interchange/Runtime/Source/Import/Private/Mesh/InterchangePhysicsAssetFactory.cpp
+++ b/Engine/Plugins/Interchange/Runtime/Source/Import/Private/Mesh/InterchangePhysicsAssetFactory.cpp
@@ -16,7 +16,7 @@ UClass* UInterchangePhysicsAssetFactory::GetFactoryClass() const
 	return UPhysicsAsset::StaticClass();
 }
 
-UObject* UInterchangePhysicsAssetFactory::ImportAssetObject_GameThread(const FImportAssetObjectParams& Arguments)
+UObject* UInterchangePhysicsAssetFactory::BeginImportAssetObject_GameThread(const FImportAssetObjectParams& Arguments)
 {
 	UObject* PhysicsAsset = nullptr;
 
@@ -89,7 +89,7 @@ UObject* UInterchangePhysicsAssetFactory::ImportAssetObject_Async(const FImportA
 	if (!ExistingAsset)
 	{
 		//NewObject is not thread safe, the asset registry directory watcher tick on the main thread can trig before we finish initializing the UObject and will crash
-		//The UObject should have been create by calling ImportAssetObject_GameThread on the main thread.
+		//The UObject should have been create by calling BeginImportAssetObject_GameThread on the main thread.
 		if (IsInGameThread())
 		{
 			PhysicsAssetObject = NewObject<UObject>(Arguments.Parent, PhysicsAssetClass, *Arguments.AssetName, RF_Public | RF_Standalone);
diff --git a/Engine/Plugins/Interchange/Runtime/Source/Import/Private/Mesh/InterchangeSkeletalMeshFactory.cpp b/Engine/Plugins/Interchange/Runtime/Source/Import/Private/Mesh/InterchangeSkeletalMeshFactory.cpp
--- a/Engine/Plugins/Interchange/Runtime/Source/Import/Private/Mesh/InterchangeSkeletalMeshFactory.cpp
+++ b/Engine/Plugins/Interchange/Runtime/Source/Import/Private/Mesh/InterchangeSkeletalMeshFactory.cpp
@@ -40,6 +40,10 @@
 #include "SkeletalMeshAttributes.h"
 #include "SkeletalMeshOperations.h"
 
+#if WITH_EDITOR
+#include "LODUtilities.h"
+#endif
+
 #include UE_INLINE_GENERATED_CPP_BY_NAME(InterchangeSkeletalMeshFactory)
 
 #if WITH_EDITORONLY_DATA
@@ -881,9 +885,9 @@ UClass* UInterchangeSkeletalMeshFactory::GetFactoryClass() const
 	return USkeletalMesh::StaticClass();
 }
 
-UObject* UInterchangeSkeletalMeshFactory::ImportAssetObject_GameThread(const FImportAssetObjectParams& Arguments)
+UObject* UInterchangeSkeletalMeshFactory::BeginImportAssetObject_GameThread(const FImportAssetObjectParams& Arguments)
 {
-	TRACE_CPUPROFILER_EVENT_SCOPE("UInterchangeSkeletalMeshFactory::ImportAssetObject_GameThread")
+	TRACE_CPUPROFILER_EVENT_SCOPE("UInterchangeSkeletalMeshFactory::BeginImportAssetObject_GameThread")
 #if !WITH_EDITOR || !WITH_EDITORONLY_DATA
 
 	UE_LOG(LogInterchangeImport, Error, TEXT("Cannot import skeletalMesh asset in runtime, this is an editor only feature."));
@@ -1077,7 +1081,6 @@ UObject* UInterchangeSkeletalMeshFactory::ImportAssetObject_Async(const FImportA
 
 	// Update skeletal materials
 	TArray<FSkeletalMaterial>& Materials = SkeletalMesh->GetMaterials();
-
 	auto UpdateOrAddSkeletalMaterial = [&Materials, bIsReImport](const FName& MaterialSlotName, UMaterialInterface* MaterialInterface)
 	{
 		UMaterialInterface* NewMaterial = MaterialInterface ? MaterialInterface : UMaterial::GetDefaultMaterial(MD_Surface);
@@ -1126,6 +1129,9 @@ UObject* UInterchangeSkeletalMeshFactory::ImportAssetObject_Async(const FImportA
 		UpdateOrAddSkeletalMaterial(MaterialSlotName, MaterialInterface ? MaterialInterface : nullptr);
 	}
 
+	ImportAssetObjectData.bIsReImport = bIsReImport;
+	ImportAssetObjectData.bApplySkinningOnly = bApplySkinningOnly;
+	
 	for (int32 LodIndex = 0; LodIndex < LodCount; ++LodIndex)
 	{
 		TRACE_CPUPROFILER_EVENT_SCOPE("UInterchangeSkeletalMeshFactory::CreateAsset_LOD")
@@ -1186,6 +1192,7 @@ UObject* UInterchangeSkeletalMeshFactory::ImportAssetObject_Async(const FImportA
 				UE_LOG(LogInterchangeImport, Warning, TEXT("Invalid Skeleton LOD when importing SkeletalMesh asset %s"), *Arguments.AssetName);
 				break;
 			}
+			ImportAssetObjectData.SkeletonReference = SkeletonReference;
 		}
 
 		FString RootJointNodeId;
@@ -1371,6 +1378,31 @@ UObject* UInterchangeSkeletalMeshFactory::ImportAssetObject_Async(const FImportA
 			FSkeletalMeshImportData::ReplaceSkeletalMeshGeometryImportData(SkeletalMesh, &SkeletalMeshImportData, CurrentLodIndex);
 		}
 
+		//Store the existing material import data before updating it so we can remap properly the material  and section data
+		TArray<FName> ExistingOriginalPerSectionMaterialImportName;
+		if (bIsReImport)
+		{
+			if (CurrentLodIndex != 0)
+			{
+				if (!SkeletalMesh->IsLODImportedDataEmpty(CurrentLodIndex))
+				{
+					FSkeletalMeshImportData LODImportData;
+					SkeletalMesh->LoadLODImportedData(CurrentLodIndex, LODImportData);
+					for (int32 SectionIndex = 0; SectionIndex < LODImportData.Materials.Num(); ++SectionIndex)
+					{
+						ExistingOriginalPerSectionMaterialImportName.Add(FName(*LODImportData.Materials[SectionIndex].MaterialImportName));
+					}
+				}
+			}
+			else
+			{
+				//LOD 0 import data is reorder to the material array before when building the LOD 0
+				for (int32 MaterialIndex = 0; MaterialIndex < Materials.Num(); ++MaterialIndex)
+				{
+					ExistingOriginalPerSectionMaterialImportName.Add(Materials[MaterialIndex].ImportedMaterialSlotName);
+				}
+			}
+		}
 		//Store the original fbx import data the SkelMeshImportDataPtr should not be modified after this
 		SkeletalMesh->SaveLODImportedData(CurrentLodIndex, SkeletalMeshImportData);
 
@@ -1410,50 +1442,6 @@ UObject* UInterchangeSkeletalMeshFactory::ImportAssetObject_Async(const FImportA
 			AddLodInfo();
 		}
 
-		const TArray<SkeletalMeshImportData::FMaterial>& ImportedMaterials = SkeletalMeshImportData.Materials;
-		if (FSkeletalMeshLODInfo* LodInfo = SkeletalMesh->GetLODInfo(CurrentLodIndex))
-		{
-			LodInfo->LODMaterialMap.Empty();
-
-			// Now set up the material mapping array.
-			for (int32 ImportedMaterialIndex = 0; ImportedMaterialIndex < ImportedMaterials.Num(); ImportedMaterialIndex++)
-			{
-				FName ImportedMaterialName = *(ImportedMaterials[ImportedMaterialIndex].MaterialImportName);
-				//Match by name
-				int32 LODMatIndex = INDEX_NONE;
-				for (int32 MaterialIndex = 0; MaterialIndex < Materials.Num(); ++MaterialIndex)
-				{
-					const FSkeletalMaterial& SkeletalMaterial = Materials[MaterialIndex];
-					if (SkeletalMaterial.ImportedMaterialSlotName != NAME_None && SkeletalMaterial.ImportedMaterialSlotName == ImportedMaterialName)
-					{
-						LODMatIndex = MaterialIndex;
-						break;
-					}
-				}
-
-				// If we don't have a match, add a new entry to the material list.
-				if (LODMatIndex == INDEX_NONE)
-				{
-					LODMatIndex = Materials.Add(FSkeletalMaterial(ImportedMaterials[ImportedMaterialIndex].Material.Get(), true, false, ImportedMaterialName, ImportedMaterialName));
-					LodInfo->LODMaterialMap.Add(LODMatIndex);
-				}
-				else
-				{
-					if (CurrentLodIndex == 0)
-					{
-						//For base lod, the skeletalmesh build will reorder the FSkeletalMeshImportData face in the correct order
-						//see FLODUtilities::AdjustImportDataFaceMaterialIndex function, so we dont need to set 
-						LodInfo->LODMaterialMap.Add(ImportedMaterialIndex);
-					}
-					else
-					{
-						//For non base lod we want to use the LODMatIndex
-						LodInfo->LODMaterialMap.Add(LODMatIndex);
-					}
-				}
-			}
-		}
-
 		//Update the bounding box if we are importing the LOD 0
 		if(CurrentLodIndex == 0)
 		{
@@ -1467,13 +1455,88 @@ UObject* UInterchangeSkeletalMeshFactory::ImportAssetObject_Async(const FImportA
 			FBoxSphereBounds BoxSphereBound((FBox)BoundingBox);
 			SkeletalMesh->SetImportedBounds(FBoxSphereBounds((FBox)BoundingBox));
 		}
-
+		//Copy the data into the game thread structure so we can finish the import in the game thread callback
+		FImportAssetObjectLODData& ImportAssetObjectLODData = ImportAssetObjectData.LodDatas.AddDefaulted_GetRef();
+		ImportAssetObjectLODData.LodIndex = CurrentLodIndex;
+		ImportAssetObjectLODData.ExistingOriginalPerSectionMaterialImportName = ExistingOriginalPerSectionMaterialImportName;
+		ImportAssetObjectLODData.ImportedMaterials = SkeletalMeshImportData.Materials;
+		
 		CurrentLodIndex++;
 	}
 
-	if(SkeletonReference)
+	//Getting the file Hash will cache it into the source data
+	Arguments.SourceData->GetFileContentHash();
+
+	return SkeletalMeshObject;
+
+#endif //else !WITH_EDITOR || !WITH_EDITORONLY_DATA
+}
+
+UObject* UInterchangeSkeletalMeshFactory::EndImportAssetObject_GameThread(const FImportAssetObjectParams& Arguments)
+{
+	TRACE_CPUPROFILER_EVENT_SCOPE("UInterchangeSkeletalMeshFactory::EndImportAssetObject_GameThread")
+	
+	check(IsInGameThread());
+
+#if !WITH_EDITOR || !WITH_EDITORONLY_DATA
+
+	UE_LOG(LogInterchangeImport, Error, TEXT("Cannot import skeletalMesh asset in runtime, this is an editor only feature."));
+	return nullptr;
+
+#else
+	if (!Arguments.AssetNode || !Arguments.AssetNode->GetObjectClass()->IsChildOf(GetFactoryClass()))
+	{
+		return nullptr;
+	}
+
+	UInterchangeSkeletalMeshFactoryNode* SkeletalMeshFactoryNode = Cast<UInterchangeSkeletalMeshFactoryNode>(Arguments.AssetNode);
+	if (SkeletalMeshFactoryNode == nullptr)
+	{
+		return nullptr;
+	}
+
+	const UClass* SkeletalMeshClass = SkeletalMeshFactoryNode->GetObjectClass();
+	check(SkeletalMeshClass && SkeletalMeshClass->IsChildOf(GetFactoryClass()));
+
+	//Get the skeletal mesh asset
+	USkeletalMesh* SkeletalMesh = Cast<USkeletalMesh>(StaticFindObject(nullptr, Arguments.Parent, *Arguments.AssetName));
+	if (!ensure(SkeletalMesh))
 	{
-		if ((!bApplySkinningOnly || !bIsReImport) && !SkeletonReference->MergeAllBonesToBoneTree(SkeletalMesh))
+		if (Arguments.ReimportObject == nullptr)
+		{
+			UE_LOG(LogInterchangeImport, Error, TEXT("Could not create skeletalMesh asset %s"), *Arguments.AssetName);
+		}
+		else
+		{
+			UE_LOG(LogInterchangeImport, Error, TEXT("Could not find reimported skeletalMesh asset %s"), *Arguments.AssetName);
+		}
+		return nullptr;
+	}
+
+	//Make sure we can modify the skeletalmesh properties
+	FSkinnedAssetAsyncBuildScope AsyncBuildScope(SkeletalMesh);
+
+	//Finish the import in the game thread so we can pop dialog if needed
+	if (ensure(ImportAssetObjectData.IsValid()))
+	{
+		for (FImportAssetObjectLODData& ImportAssetObjectLODData : ImportAssetObjectData.LodDatas)
+		{
+			if (FSkeletalMeshLODInfo* LodInfo = SkeletalMesh->GetLODInfo(ImportAssetObjectLODData.LodIndex))
+			{
+				FLODUtilities::FSkeletalMeshMatchImportedMaterialsParameters Parameters;
+				Parameters.bIsReImport = ImportAssetObjectData.bIsReImport;
+				Parameters.LodIndex = ImportAssetObjectLODData.LodIndex;
+				Parameters.SkeletalMesh = SkeletalMesh;
+				Parameters.ImportedMaterials = &ImportAssetObjectLODData.ImportedMaterials;
+				Parameters.ExistingOriginalPerSectionMaterialImportName = &ImportAssetObjectLODData.ExistingOriginalPerSectionMaterialImportName;
+				FLODUtilities::MatchImportedMaterials(Parameters);
+			}
+		}
+	}
+
+	if (USkeleton* SkeletonReference = ImportAssetObjectData.SkeletonReference)
+	{
+		if ((!ImportAssetObjectData.bApplySkinningOnly || !ImportAssetObjectData.bIsReImport) && !SkeletonReference->MergeAllBonesToBoneTree(SkeletalMesh))
 		{
 			TUniqueFunction<bool()> RecreateSkeleton = [this, WeakSkeletalMesh = TWeakObjectPtr<USkeletalMesh>(SkeletalMesh), WeakSkeleton = TWeakObjectPtr<USkeleton>(SkeletonReference)]()
 			{
@@ -1491,7 +1554,7 @@ UObject* UInterchangeSkeletalMeshFactory::ImportAssetObject_Async(const FImportA
 				if (GIsRunningUnattendedScript)
 				{
 					UInterchangeResultError_Generic* Message = AddMessage<UInterchangeResultError_Generic>();
-					Message->Text = NSLOCTEXT("InterchangeSkeletalMeshFactory", "ImportWithScriptIncompatibleSkeleton", "Interchange Import UInterchangeSkeletalMeshFactory::CreateAsset, cannot merge bone tree with the existing skeleton.");
+					Message->Text = NSLOCTEXT("InterchangeSkeletalMeshFactory", "ImportWithScriptIncompatibleSkeleton", "Interchange Import UInterchangeSkeletalMeshFactory::EndImportAssetObject_GameThread, cannot merge bone tree with the existing skeleton.");
 					return false;
 				}
 
@@ -1570,10 +1633,10 @@ UObject* UInterchangeSkeletalMeshFactory::ImportAssetObject_Async(const FImportA
 	}
 	else
 	{
-		UE_LOG(LogInterchangeImport, Error, TEXT("Interchange Import UInterchangeSkeletalMeshFactory::CreateAsset, USkeleton* SkeletonReference is nullptr."));
+		UE_LOG(LogInterchangeImport, Error, TEXT("Interchange Import UInterchangeSkeletalMeshFactory::EndImportAssetObject_GameThread, USkeleton* SkeletonReference is nullptr."));
 	}
 
-	if (bIsReImport)
+	if (ImportAssetObjectData.bIsReImport)
 	{
 		//We must reset the matrixs since CalculateInvRefMatrices only do the calculation if matrix count differ from the bone count.
 		SkeletalMesh->GetRefBasesInvMatrix().Reset();
@@ -1581,7 +1644,7 @@ UObject* UInterchangeSkeletalMeshFactory::ImportAssetObject_Async(const FImportA
 
 	SkeletalMesh->CalculateInvRefMatrices();
 
-	if (!bIsReImport)
+	if (!ImportAssetObjectData.bIsReImport)
 	{
 		/** Apply all SkeletalMeshFactoryNode custom attributes to the skeletal mesh asset */
 		SkeletalMeshFactoryNode->ApplyAllCustomAttributeToObject(SkeletalMesh);
@@ -1614,14 +1677,9 @@ UObject* UInterchangeSkeletalMeshFactory::ImportAssetObject_Async(const FImportA
 		CurrentNode->FillAllCustomAttributeFromObject(SkeletalMesh);
 		UE::Interchange::FFactoryCommon::ApplyReimportStrategyToAsset(SkeletalMesh, PreviousNode, CurrentNode, SkeletalMeshFactoryNode);
 	}
-		
-	//Getting the file Hash will cache it into the source data
-	Arguments.SourceData->GetFileContentHash();
-
-	//The interchange completion task (call in the GameThread after the factories pass), will call PostEditChange which will trig another asynchronous system that will build all material in parallel
-	return SkeletalMeshObject;
 
-#endif //else !WITH_EDITOR || !WITH_EDITORONLY_DATA
+	return SkeletalMesh;
+#endif //WITH_EDITOR
 }
 
 void UInterchangeSkeletalMeshFactory::Cancel()
@@ -1760,3 +1818,25 @@ bool UInterchangeSkeletalMeshFactory::SetReimportSourceIndex(const UObject* Obje
 	return false;
 }
 
+bool UInterchangeSkeletalMeshFactory::FImportAssetObjectData::IsValid() const
+{
+	if (LodDatas.Num() < 1)
+	{
+		return false;
+	}
+	for (const FImportAssetObjectLODData& ImportAssetObjectLODData : LodDatas)
+	{
+		if (ImportAssetObjectLODData.LodIndex == INDEX_NONE)
+		{
+			return false;
+		}
+#if WITH_EDITOR
+		if (ImportAssetObjectLODData.ImportedMaterials.Num() < 1)
+		{
+			return false;
+		}
+#endif
+	}
+
+	return true;
+}
\ No newline at end of file
diff --git a/Engine/Plugins/Interchange/Runtime/Source/Import/Private/Mesh/InterchangeSkeletonFactory.cpp b/Engine/Plugins/Interchange/Runtime/Source/Import/Private/Mesh/InterchangeSkeletonFactory.cpp
--- a/Engine/Plugins/Interchange/Runtime/Source/Import/Private/Mesh/InterchangeSkeletonFactory.cpp
+++ b/Engine/Plugins/Interchange/Runtime/Source/Import/Private/Mesh/InterchangeSkeletonFactory.cpp
@@ -17,7 +17,7 @@ UClass* UInterchangeSkeletonFactory::GetFactoryClass() const
 	return USkeleton::StaticClass();
 }
 
-UObject* UInterchangeSkeletonFactory::ImportAssetObject_GameThread(const FImportAssetObjectParams& Arguments)
+UObject* UInterchangeSkeletonFactory::BeginImportAssetObject_GameThread(const FImportAssetObjectParams& Arguments)
 {
 	UObject* Skeleton = nullptr;
 
diff --git a/Engine/Plugins/Interchange/Runtime/Source/Import/Private/Mesh/InterchangeStaticMeshFactory.cpp b/Engine/Plugins/Interchange/Runtime/Source/Import/Private/Mesh/InterchangeStaticMeshFactory.cpp
--- a/Engine/Plugins/Interchange/Runtime/Source/Import/Private/Mesh/InterchangeStaticMeshFactory.cpp
+++ b/Engine/Plugins/Interchange/Runtime/Source/Import/Private/Mesh/InterchangeStaticMeshFactory.cpp
@@ -44,14 +44,107 @@
 #include "EditorFramework/AssetImportData.h"
 #endif //WITH_EDITORONLY_DATA
 
+static bool GInterchangeStaticMeshReorderMaterialSlots = true;
+static FAutoConsoleVariableRef CCvarInterchangeStaticMeshReorderMaterialSlots(
+	TEXT("Interchange.FeatureFlags.Import.StaticMesh.ReorderMaterialSlots"),
+	GInterchangeStaticMeshReorderMaterialSlots,
+	TEXT("Whether Re-importing a static mesh should reorder the material slots."),
+	ECVF_Default);
 
 UClass* UInterchangeStaticMeshFactory::GetFactoryClass() const
 {
 	return UStaticMesh::StaticClass();
 }
 
+namespace UE::Interchange::Private::StaticMesh
+{
+	void ReorderMaterialSlotToBaseLod(UStaticMesh* StaticMesh)
+	{
+#if WITH_EDITOR
+		if (!StaticMesh || !StaticMesh->IsMeshDescriptionValid(0))
+		{
+			return;
+		}
 
-UObject* UInterchangeStaticMeshFactory::ImportAssetObject_GameThread(const FImportAssetObjectParams& Arguments)
+		TArray<FStaticMaterial>& Materials = StaticMesh->GetStaticMaterials();
+		if (Materials.Num() < 2)
+		{
+			return;
+		}
+
+		const FMeshDescription& BaseMeshDescription = *StaticMesh->GetMeshDescription(0);
+		if (!ensure(!BaseMeshDescription.NeedsCompact()))
+		{
+			return;
+		}
+		
+		FStaticMeshConstAttributes StaticMeshAttributes(BaseMeshDescription);
+		TPolygonGroupAttributesRef<const FName> SlotNames = StaticMeshAttributes.GetPolygonGroupMaterialSlotNames();
+
+		TArray<int32> RemapMaterialIndexes;
+		RemapMaterialIndexes.Reserve(Materials.Num());
+		for (int32 MaterialIndex = 0; MaterialIndex < Materials.Num(); ++MaterialIndex)
+		{
+			RemapMaterialIndexes.Add(INDEX_NONE);
+		}
+		TArray<FStaticMaterial> ReorderMaterialArray;
+		ReorderMaterialArray.Reserve(Materials.Num());
+		for (FPolygonGroupID PolygonGroupID : BaseMeshDescription.PolygonGroups().GetElementIDs())
+		{
+			FName ImportMaterialName = SlotNames[PolygonGroupID];
+			bool bFoundMatch = false;
+			for (int32 MaterialIndex = 0; MaterialIndex < Materials.Num(); ++MaterialIndex)
+			{
+				if (RemapMaterialIndexes[MaterialIndex] != INDEX_NONE)
+				{
+					//This material was already matched
+					continue;
+				}
+				int32& RemapIndex = RemapMaterialIndexes[MaterialIndex];
+				FName MaterialName = Materials[MaterialIndex].ImportedMaterialSlotName;
+				if (MaterialName == ImportMaterialName)
+				{
+					RemapIndex = ReorderMaterialArray.Add(Materials[MaterialIndex]);
+					bFoundMatch = true;
+					break;
+				}
+			}
+			//All mesh description polygon group should have a match
+			ensure(bFoundMatch);
+		}
+
+		//Custom LOD can add materials, so we add them at the end of the material slots
+		for (int32 MaterialIndex = 0; MaterialIndex < Materials.Num(); ++MaterialIndex)
+		{
+			if (RemapMaterialIndexes[MaterialIndex] == INDEX_NONE)
+			{
+				RemapMaterialIndexes[MaterialIndex] = ReorderMaterialArray.Add(Materials[MaterialIndex]);
+			}
+		}
+
+		check(ReorderMaterialArray.Num() == Materials.Num());
+
+		//Reorder the static mesh material slot array
+		Materials = ReorderMaterialArray;
+
+		//Fix all the section info map MaterialIndex with the RemapMaterialIndexes
+		FMeshSectionInfoMap& SectionInfoMap = StaticMesh->GetSectionInfoMap();
+		const int32 LodCount = StaticMesh->GetNumSourceModels();
+		for (int32 LodIndex = 0; LodIndex < LodCount; ++LodIndex)
+		{
+			const int32 SectionCount = SectionInfoMap.GetSectionNumber(LodIndex);
+			for (int32 SectionIndex = 0; SectionIndex < SectionCount; ++SectionIndex)
+			{
+				FMeshSectionInfo SectionInfo = SectionInfoMap.Get(LodIndex, SectionIndex);
+				SectionInfo.MaterialIndex = RemapMaterialIndexes[SectionInfo.MaterialIndex];
+				SectionInfoMap.Set(LodIndex, SectionIndex, SectionInfo);
+			}
+		}
+#endif // WITH_EDITOR
+	}
+} //ns UE::Interchange::Private::StaticMesh
+
+UObject* UInterchangeStaticMeshFactory::BeginImportAssetObject_GameThread(const FImportAssetObjectParams& Arguments)
 {
 	UStaticMesh* StaticMesh = nullptr;
 	if (!Arguments.AssetNode || !Arguments.AssetNode->GetObjectClass()->IsChildOf(GetFactoryClass()))
@@ -172,10 +265,9 @@ UObject* UInterchangeStaticMeshFactory::ImportAssetObject_Async(const FImportAss
 	ensure(!StaticMesh->AreRenderingResourcesInitialized());
 
 	const int32 LodCount = StaticMeshFactoryNode->GetLodDataCount();
-	const int32 PrevLodCount = StaticMesh->GetNumLODs();
-	const int32 FinalLodCount = FMath::Max(PrevLodCount, LodCount);
-
 #if WITH_EDITOR
+	const int32 PrevLodCount = StaticMesh->GetNumSourceModels();
+	const int32 FinalLodCount = FMath::Max(PrevLodCount, LodCount);
 	StaticMesh->SetNumSourceModels(FinalLodCount);
 #endif // WITH_EDITOR
 
@@ -239,7 +331,7 @@ UObject* UInterchangeStaticMeshFactory::ImportAssetObject_Async(const FImportAss
 	StaticMeshFactoryNode->GetLodDataUniqueIds(LodDataUniqueIds);
 	ensure(LodDataUniqueIds.Num() == LodCount);
 
-	TArray<FMeshDescription> LodMeshDescriptions;
+	TArray<FMeshDescription>& LodMeshDescriptions = ImportAssetObjectData.LodMeshDescriptions;
 	LodMeshDescriptions.SetNum(LodCount);
 
 	bool bImportedCustomCollision = false;
@@ -364,33 +456,6 @@ UObject* UInterchangeStaticMeshFactory::ImportAssetObject_Async(const FImportAss
 			}
 		}
 
-		// Build section info map from materials
-		FStaticMeshConstAttributes StaticMeshAttributes(LodMeshDescription);
-		TPolygonGroupAttributesRef<const FName> SlotNames = StaticMeshAttributes.GetPolygonGroupMaterialSlotNames();
-		int32 SectionIndex = 0;
-		for (FPolygonGroupID PolygonGroupID : LodMeshDescription.PolygonGroups().GetElementIDs())
-		{
-			int32 MaterialSlotIndex = StaticMesh->GetMaterialIndexFromImportedMaterialSlotName(SlotNames[PolygonGroupID]);
-					
-			// If no material was found with this slot name, fill out a blank slot instead.
-			if (MaterialSlotIndex == INDEX_NONE)
-			{
-				MaterialSlotIndex = StaticMesh->GetStaticMaterials().Emplace(UMaterial::GetDefaultMaterial(MD_Surface), SlotNames[PolygonGroupID]);
-#if !WITH_EDITOR
-				StaticMesh->GetStaticMaterials()[MaterialSlotIndex].UVChannelData = FMeshUVChannelInfo(1.f);
-#endif
-			}
-
-#if WITH_EDITOR
-			FMeshSectionInfo Info = StaticMesh->GetSectionInfoMap().Get(CurrentLodIndex, SectionIndex);
-			Info.MaterialIndex = MaterialSlotIndex;
-			StaticMesh->GetSectionInfoMap().Remove(CurrentLodIndex, SectionIndex);
-			StaticMesh->GetSectionInfoMap().Set(CurrentLodIndex, SectionIndex, Info);
-#endif
-
-			SectionIndex++;
-		}
-
 		// Import collision geometry
 		if (CurrentLodIndex == 0)
 		{
@@ -419,7 +484,140 @@ UObject* UInterchangeStaticMeshFactory::ImportAssetObject_Async(const FImportAss
 	}
 #endif // WITH_EDITOR
 
-	CommitMeshDescriptions(*StaticMesh, MoveTemp(LodMeshDescriptions));
+	ImportAssetObjectData.bImportedCustomCollision = bImportedCustomCollision;
+
+	// Getting the file Hash will cache it into the source data
+	Arguments.SourceData->GetFileContentHash();
+
+	return StaticMeshObject;
+}
+
+UObject* UInterchangeStaticMeshFactory::EndImportAssetObject_GameThread(const FImportAssetObjectParams& Arguments)
+{
+	if (!Arguments.AssetNode || !Arguments.AssetNode->GetObjectClass()->IsChildOf(GetFactoryClass()))
+	{
+		return nullptr;
+	}
+
+	UInterchangeStaticMeshFactoryNode* StaticMeshFactoryNode = Cast<UInterchangeStaticMeshFactoryNode>(Arguments.AssetNode);
+	if (StaticMeshFactoryNode == nullptr)
+	{
+		return nullptr;
+	}
+
+	const UClass* StaticMeshClass = StaticMeshFactoryNode->GetObjectClass();
+	check(StaticMeshClass && StaticMeshClass->IsChildOf(GetFactoryClass()));
+
+	// create an asset if it doesn't exist
+	UObject* ExistingAsset = StaticFindObject(nullptr, Arguments.Parent, *Arguments.AssetName);
+
+	const bool bReimport = Arguments.ReimportObject && ExistingAsset;
+
+	UStaticMesh* StaticMesh = Cast<UStaticMesh>(ExistingAsset);
+	if (!ensure(StaticMesh))
+	{
+		UE_LOG(LogInterchangeImport, Error, TEXT("Could not create StaticMesh asset %s"), *Arguments.AssetName);
+		return nullptr;
+	}
+
+	for (int32 LodIndex = 0; LodIndex < ImportAssetObjectData.LodMeshDescriptions.Num(); ++LodIndex)
+	{
+		// Add the lod mesh data to the static mesh
+		FMeshDescription& LodMeshDescription = ImportAssetObjectData.LodMeshDescriptions[LodIndex];
+		if (LodMeshDescription.IsEmpty())
+		{
+			//All the valid Mesh description are at the beginning of the array
+			break;
+		}
+
+		// Build section info map from materials
+		FStaticMeshConstAttributes StaticMeshAttributes(LodMeshDescription);
+		TPolygonGroupAttributesRef<const FName> SlotNames = StaticMeshAttributes.GetPolygonGroupMaterialSlotNames();
+#if WITH_EDITOR
+		if (bReimport)
+		{
+			//Match the existing section info map data
+
+			//First find the old mesh description polygon groups name that match with the imported mesh description polygon groups name.
+			const int32 PreviousSectionCount = StaticMesh->GetSectionInfoMap().GetSectionNumber(LodIndex);
+			TMap<FPolygonGroupID, FPolygonGroupID> ImportedToOldPolygonGroupMatch;
+			ImportedToOldPolygonGroupMatch.Reserve(LodMeshDescription.PolygonGroups().Num());
+			if (StaticMesh->IsMeshDescriptionValid(LodIndex))
+			{
+				//Match incoming mesh description with the old mesh description
+				FMeshDescription& OldMeshDescription = *StaticMesh->GetMeshDescription(LodIndex);
+				FStaticMeshConstAttributes OldStaticMeshAttributes(OldMeshDescription);
+				TPolygonGroupAttributesRef<const FName> OldSlotNames = OldStaticMeshAttributes.GetPolygonGroupMaterialSlotNames();
+				for (FPolygonGroupID PolygonGroupID : LodMeshDescription.PolygonGroups().GetElementIDs())
+				{
+					for (FPolygonGroupID OldPolygonGroupID : OldMeshDescription.PolygonGroups().GetElementIDs())
+					{
+						if (SlotNames[PolygonGroupID] == OldSlotNames[OldPolygonGroupID])
+						{
+							ImportedToOldPolygonGroupMatch.FindOrAdd(PolygonGroupID) = OldPolygonGroupID;
+							break;
+						}
+					}
+				}
+			}
+			//Create a new set of mesh section info for this lod
+			TArray<FMeshSectionInfo> NewSectionInfoMapData;
+			NewSectionInfoMapData.Reserve(LodMeshDescription.PolygonGroups().Num());
+			for (FPolygonGroupID PolygonGroupID : LodMeshDescription.PolygonGroups().GetElementIDs())
+			{
+				if (FPolygonGroupID* OldPolygonGroupID = ImportedToOldPolygonGroupMatch.Find(PolygonGroupID))
+				{
+					if (StaticMesh->GetSectionInfoMap().IsValidSection(LodIndex, OldPolygonGroupID->GetValue()))
+					{
+						NewSectionInfoMapData.Add(StaticMesh->GetSectionInfoMap().Get(LodIndex, OldPolygonGroupID->GetValue()));
+					}
+				}
+				else
+				{
+					//This is an unmatched section, its either added or we did not recover the name
+					int32 MaterialSlotIndex = StaticMesh->GetMaterialIndexFromImportedMaterialSlotName(SlotNames[PolygonGroupID]);
+					//Missing material slot should have been added before
+					ensure(MaterialSlotIndex != INDEX_NONE);
+					NewSectionInfoMapData.Add(FMeshSectionInfo(MaterialSlotIndex));
+				}
+			}
+			//Recreate the section info map
+			for (int32 NewSectionIndex = 0; NewSectionIndex < NewSectionInfoMapData.Num(); ++NewSectionIndex)
+			{
+				StaticMesh->GetSectionInfoMap().Remove(LodIndex, NewSectionIndex);
+				StaticMesh->GetSectionInfoMap().Set(LodIndex, NewSectionIndex, NewSectionInfoMapData[NewSectionIndex]);
+			}
+		}
+		else
+#endif // WITH_EDITOR
+		{
+			int32 SectionIndex = 0;
+			for (FPolygonGroupID PolygonGroupID : LodMeshDescription.PolygonGroups().GetElementIDs())
+			{
+				int32 MaterialSlotIndex = StaticMesh->GetMaterialIndexFromImportedMaterialSlotName(SlotNames[PolygonGroupID]);
+
+				// If no material was found with this slot name, fill out a blank slot instead.
+				if (MaterialSlotIndex == INDEX_NONE)
+				{
+					MaterialSlotIndex = StaticMesh->GetStaticMaterials().Emplace(UMaterial::GetDefaultMaterial(MD_Surface), SlotNames[PolygonGroupID]);
+#if !WITH_EDITOR
+					StaticMesh->GetStaticMaterials()[MaterialSlotIndex].UVChannelData = FMeshUVChannelInfo(1.f);
+#endif
+				}
+
+#if WITH_EDITOR
+				FMeshSectionInfo Info = StaticMesh->GetSectionInfoMap().Get(LodIndex, SectionIndex);
+				Info.MaterialIndex = MaterialSlotIndex;
+				StaticMesh->GetSectionInfoMap().Remove(LodIndex, SectionIndex);
+				StaticMesh->GetSectionInfoMap().Set(LodIndex, SectionIndex, Info);
+#endif
+
+				SectionIndex++;
+			}
+		}
+	}
+
+	CommitMeshDescriptions(*StaticMesh, MoveTemp(ImportAssetObjectData.LodMeshDescriptions));
 
 	ImportSockets(Arguments, StaticMesh, StaticMeshFactoryNode);
 
@@ -445,7 +643,7 @@ UObject* UInterchangeStaticMeshFactory::ImportAssetObject_Async(const FImportAss
 	}
 #endif // WITH_EDITOR
 
-	if (!bImportedCustomCollision)
+	if (!ImportAssetObjectData.bImportedCustomCollision)
 	{
 		GenerateKDopCollision(Arguments, StaticMesh);
 	}
@@ -455,15 +653,15 @@ UObject* UInterchangeStaticMeshFactory::ImportAssetObject_Async(const FImportAss
 		StaticMesh->bCustomizedCollision = true;
 	}
 #endif // WITH_EDITORONLY_DATA
-
-	// Getting the file Hash will cache it into the source data
-	Arguments.SourceData->GetFileContentHash();
-
 #if WITH_EDITOR
 	FMeshBudgetProjectSettingsUtils::SetLodGroupForStaticMesh(StaticMesh);
 #endif
 
-	return StaticMeshObject;
+	if (bReimport && GInterchangeStaticMeshReorderMaterialSlots)
+	{
+		UE::Interchange::Private::StaticMesh::ReorderMaterialSlotToBaseLod(StaticMesh);
+	}
+	return StaticMesh;
 }
 
 void UInterchangeStaticMeshFactory::CommitMeshDescriptions(UStaticMesh& StaticMesh, TArray<FMeshDescription>&& LodMeshDescriptions)
diff --git a/Engine/Plugins/Interchange/Runtime/Source/Import/Private/Scene/InterchangeSceneVariantSetsFactory.cpp b/Engine/Plugins/Interchange/Runtime/Source/Import/Private/Scene/InterchangeSceneVariantSetsFactory.cpp
--- a/Engine/Plugins/Interchange/Runtime/Source/Import/Private/Scene/InterchangeSceneVariantSetsFactory.cpp
+++ b/Engine/Plugins/Interchange/Runtime/Source/Import/Private/Scene/InterchangeSceneVariantSetsFactory.cpp
@@ -294,7 +294,7 @@ UClass* UInterchangeSceneVariantSetsFactory::GetFactoryClass() const
 	return ULevelVariantSets::StaticClass();
 }
 
-UObject* UInterchangeSceneVariantSetsFactory::ImportAssetObject_GameThread(const FImportAssetObjectParams& Arguments)
+UObject* UInterchangeSceneVariantSetsFactory::BeginImportAssetObject_GameThread(const FImportAssetObjectParams& Arguments)
 {
 #if !WITH_EDITOR || !WITH_EDITORONLY_DATA
 
diff --git a/Engine/Plugins/Interchange/Runtime/Source/Import/Private/Texture/InterchangeTextureFactory.cpp b/Engine/Plugins/Interchange/Runtime/Source/Import/Private/Texture/InterchangeTextureFactory.cpp
--- a/Engine/Plugins/Interchange/Runtime/Source/Import/Private/Texture/InterchangeTextureFactory.cpp
+++ b/Engine/Plugins/Interchange/Runtime/Source/Import/Private/Texture/InterchangeTextureFactory.cpp
@@ -1459,7 +1459,7 @@ UClass* UInterchangeTextureFactory::GetFactoryClass() const
 	return UTexture::StaticClass();
 }
 
-UObject* UInterchangeTextureFactory::ImportAssetObject_GameThread(const FImportAssetObjectParams& Arguments)
+UObject* UInterchangeTextureFactory::BeginImportAssetObject_GameThread(const FImportAssetObjectParams& Arguments)
 {
 	using namespace  UE::Interchange::Private::InterchangeTextureFactory;
 
diff --git a/Engine/Plugins/Interchange/Runtime/Source/Import/Public/Animation/InterchangeAnimSequenceFactory.h b/Engine/Plugins/Interchange/Runtime/Source/Import/Public/Animation/InterchangeAnimSequenceFactory.h
--- a/Engine/Plugins/Interchange/Runtime/Source/Import/Public/Animation/InterchangeAnimSequenceFactory.h
+++ b/Engine/Plugins/Interchange/Runtime/Source/Import/Public/Animation/InterchangeAnimSequenceFactory.h
@@ -23,7 +23,7 @@ public:
 
 	virtual UClass* GetFactoryClass() const override;
 	virtual EInterchangeFactoryAssetType GetFactoryAssetType() override { return EInterchangeFactoryAssetType::Animations; }
-	virtual UObject* ImportAssetObject_GameThread(const FImportAssetObjectParams& Arguments) override;
+	virtual UObject* BeginImportAssetObject_GameThread(const FImportAssetObjectParams& Arguments) override;
 	virtual void SetupObject_GameThread(const FSetupObjectParams& Arguments) override;
 
 	virtual bool GetSourceFilenames(const UObject* Object, TArray<FString>& OutSourceFilenames) const override;
diff --git a/Engine/Plugins/Interchange/Runtime/Source/Import/Public/Animation/InterchangeAnimationTrackSetFactory.h b/Engine/Plugins/Interchange/Runtime/Source/Import/Public/Animation/InterchangeAnimationTrackSetFactory.h
--- a/Engine/Plugins/Interchange/Runtime/Source/Import/Public/Animation/InterchangeAnimationTrackSetFactory.h
+++ b/Engine/Plugins/Interchange/Runtime/Source/Import/Public/Animation/InterchangeAnimationTrackSetFactory.h
@@ -20,7 +20,7 @@ public:
 
 	virtual UClass* GetFactoryClass() const override;
 	virtual EInterchangeFactoryAssetType GetFactoryAssetType() override { return EInterchangeFactoryAssetType::Animations; }
-	virtual UObject* ImportAssetObject_GameThread(const FImportAssetObjectParams& Arguments) override;
+	virtual UObject* BeginImportAssetObject_GameThread(const FImportAssetObjectParams& Arguments) override;
 	virtual void SetupObject_GameThread(const FSetupObjectParams& Arguments) override;
 
 	// Interchange factory base interface end
diff --git a/Engine/Plugins/Interchange/Runtime/Source/Import/Public/Material/InterchangeMaterialFactory.h b/Engine/Plugins/Interchange/Runtime/Source/Import/Public/Material/InterchangeMaterialFactory.h
--- a/Engine/Plugins/Interchange/Runtime/Source/Import/Public/Material/InterchangeMaterialFactory.h
+++ b/Engine/Plugins/Interchange/Runtime/Source/Import/Public/Material/InterchangeMaterialFactory.h
@@ -28,7 +28,7 @@ public:
 
 	virtual UClass* GetFactoryClass() const override;
 	virtual EInterchangeFactoryAssetType GetFactoryAssetType() override { return EInterchangeFactoryAssetType::Materials; }
-	virtual UObject* ImportAssetObject_GameThread(const FImportAssetObjectParams& Arguments) override;
+	virtual UObject* BeginImportAssetObject_GameThread(const FImportAssetObjectParams& Arguments) override;
 	virtual UObject* ImportAssetObject_Async(const FImportAssetObjectParams& Arguments) override;
 	virtual void SetupObject_GameThread(const FSetupObjectParams& Arguments) override;
 	virtual bool GetSourceFilenames(const UObject* Object, TArray<FString>& OutSourceFilenames) const override;
@@ -55,7 +55,7 @@ public:
 
 	virtual UClass* GetFactoryClass() const override;
 	virtual EInterchangeFactoryAssetType GetFactoryAssetType() override { return EInterchangeFactoryAssetType::Materials; }
-	virtual UObject* ImportAssetObject_GameThread(const FImportAssetObjectParams& Arguments) override;
+	virtual UObject* BeginImportAssetObject_GameThread(const FImportAssetObjectParams& Arguments) override;
 	virtual UObject* ImportAssetObject_Async(const FImportAssetObjectParams& Arguments) override;
 	virtual void SetupObject_GameThread(const FSetupObjectParams& Arguments) override;
 
diff --git a/Engine/Plugins/Interchange/Runtime/Source/Import/Public/Mesh/InterchangePhysicsAssetFactory.h b/Engine/Plugins/Interchange/Runtime/Source/Import/Public/Mesh/InterchangePhysicsAssetFactory.h
--- a/Engine/Plugins/Interchange/Runtime/Source/Import/Public/Mesh/InterchangePhysicsAssetFactory.h
+++ b/Engine/Plugins/Interchange/Runtime/Source/Import/Public/Mesh/InterchangePhysicsAssetFactory.h
@@ -20,7 +20,7 @@ public:
 
 	virtual UClass* GetFactoryClass() const override;
 	virtual EInterchangeFactoryAssetType GetFactoryAssetType() override { return EInterchangeFactoryAssetType::Physics; }
-	virtual UObject* ImportAssetObject_GameThread(const FImportAssetObjectParams& Arguments) override;
+	virtual UObject* BeginImportAssetObject_GameThread(const FImportAssetObjectParams& Arguments) override;
 	virtual UObject* ImportAssetObject_Async(const FImportAssetObjectParams& Arguments) override;
 
 	// Interchange factory base interface end
diff --git a/Engine/Plugins/Interchange/Runtime/Source/Import/Public/Mesh/InterchangeSkeletalMeshFactory.h b/Engine/Plugins/Interchange/Runtime/Source/Import/Public/Mesh/InterchangeSkeletalMeshFactory.h
--- a/Engine/Plugins/Interchange/Runtime/Source/Import/Public/Mesh/InterchangeSkeletalMeshFactory.h
+++ b/Engine/Plugins/Interchange/Runtime/Source/Import/Public/Mesh/InterchangeSkeletalMeshFactory.h
@@ -4,26 +4,45 @@
 
 #include "CoreMinimal.h"
 #include "InterchangeFactoryBase.h"
+#include "Rendering/SkeletalMeshLODImporterData.h"
 #include "UObject/Object.h"
 #include "UObject/ObjectMacros.h"
 
 #include "InterchangeSkeletalMeshFactory.generated.h"
 
+class UInterchangeSkeletalMeshFactoryNode;
 class USkeletalMesh;
+class USkeleton;
 
 UCLASS(BlueprintType, Experimental)
 class INTERCHANGEIMPORT_API UInterchangeSkeletalMeshFactory : public UInterchangeFactoryBase
 {
 	GENERATED_BODY()
 public:
-
+	struct FImportAssetObjectLODData
+	{
+		int32 LodIndex = INDEX_NONE;
+		TArray<FName> ExistingOriginalPerSectionMaterialImportName;
+#if WITH_EDITOR
+		TArray<SkeletalMeshImportData::FMaterial> ImportedMaterials;
+#endif
+	};
+	struct FImportAssetObjectData
+	{
+		bool bIsReImport = false;
+		USkeleton* SkeletonReference = nullptr;
+		bool bApplySkinningOnly = false;
+		TArray<FImportAssetObjectLODData> LodDatas;
+		bool IsValid() const;
+	};
 	//////////////////////////////////////////////////////////////////////////
 	// Interchange factory base interface begin
 
 	virtual UClass* GetFactoryClass() const override;
 	virtual EInterchangeFactoryAssetType GetFactoryAssetType() override { return EInterchangeFactoryAssetType::Meshes; }
-	virtual UObject* ImportAssetObject_GameThread(const FImportAssetObjectParams& Arguments) override;
+	virtual UObject* BeginImportAssetObject_GameThread(const FImportAssetObjectParams& Arguments) override;
 	virtual UObject* ImportAssetObject_Async(const FImportAssetObjectParams& Arguments) override;
+	virtual UObject* EndImportAssetObject_GameThread(const FImportAssetObjectParams& Arguments) override;
 	virtual void Cancel() override;
 	virtual void SetupObject_GameThread(const FSetupObjectParams& Arguments) override;
 	virtual bool GetSourceFilenames(const UObject* Object, TArray<FString>& OutSourceFilenames) const override;
@@ -34,6 +53,8 @@ public:
 	//////////////////////////////////////////////////////////////////////////
 private:
 	FEvent* SkeletalMeshLockPropertiesEvent = nullptr;
+
+	FImportAssetObjectData ImportAssetObjectData;
 };
 
 
diff --git a/Engine/Plugins/Interchange/Runtime/Source/Import/Public/Mesh/InterchangeSkeletonFactory.h b/Engine/Plugins/Interchange/Runtime/Source/Import/Public/Mesh/InterchangeSkeletonFactory.h
--- a/Engine/Plugins/Interchange/Runtime/Source/Import/Public/Mesh/InterchangeSkeletonFactory.h
+++ b/Engine/Plugins/Interchange/Runtime/Source/Import/Public/Mesh/InterchangeSkeletonFactory.h
@@ -20,7 +20,7 @@ public:
 
 	virtual UClass* GetFactoryClass() const override;
 	virtual EInterchangeFactoryAssetType GetFactoryAssetType() override { return EInterchangeFactoryAssetType::Meshes; }
-	virtual UObject* ImportAssetObject_GameThread(const FImportAssetObjectParams& Arguments) override;
+	virtual UObject* BeginImportAssetObject_GameThread(const FImportAssetObjectParams& Arguments) override;
 	virtual UObject* ImportAssetObject_Async(const FImportAssetObjectParams& Arguments) override;
 
 	// Interchange factory base interface end
diff --git a/Engine/Plugins/Interchange/Runtime/Source/Import/Public/Mesh/InterchangeStaticMeshFactory.h b/Engine/Plugins/Interchange/Runtime/Source/Import/Public/Mesh/InterchangeStaticMeshFactory.h
--- a/Engine/Plugins/Interchange/Runtime/Source/Import/Public/Mesh/InterchangeStaticMeshFactory.h
+++ b/Engine/Plugins/Interchange/Runtime/Source/Import/Public/Mesh/InterchangeStaticMeshFactory.h
@@ -4,6 +4,7 @@
 
 #include "CoreMinimal.h"
 #include "InterchangeFactoryBase.h"
+#include "MeshDescription.h"
 #include "UObject/Object.h"
 #include "UObject/ObjectMacros.h"
 #include "Mesh/InterchangeMeshPayload.h"
@@ -15,7 +16,6 @@ class UBodySetup;
 class UStaticMesh;
 class UInterchangeStaticMeshFactoryNode;
 class UInterchangeStaticMeshLodDataNode;
-struct FMeshDescription;
 struct FKAggregateGeom;
 
 
@@ -30,8 +30,9 @@ public:
 
 	virtual UClass* GetFactoryClass() const override;
 	virtual EInterchangeFactoryAssetType GetFactoryAssetType() override { return EInterchangeFactoryAssetType::Meshes; }
-	virtual UObject* ImportAssetObject_GameThread(const FImportAssetObjectParams& Arguments) override;
+	virtual UObject* BeginImportAssetObject_GameThread(const FImportAssetObjectParams& Arguments) override;
 	virtual UObject* ImportAssetObject_Async(const FImportAssetObjectParams& Arguments) override;
+	virtual UObject* EndImportAssetObject_GameThread(const FImportAssetObjectParams& Arguments) override;
 	virtual void SetupObject_GameThread(const FSetupObjectParams& Arguments) override;
 	virtual bool GetSourceFilenames(const UObject* Object, TArray<FString>& OutSourceFilenames) const override;
 	virtual bool SetSourceFilename(const UObject* Object, const FString& SourceFilename, int32 SourceIndex) const override;
@@ -68,4 +69,10 @@ private:
 #if WITH_EDITORONLY_DATA
 	void SetupSourceModelsSettings(UStaticMesh& StaticMesh, const TArray<FMeshDescription>& LodMeshDescriptions, int32 PreviousLodCount, int32 FinalLodCount, bool bIsAReimport);
 #endif
+	struct FImportAssetObjectData
+	{
+		bool bImportedCustomCollision = false;
+		TArray<FMeshDescription> LodMeshDescriptions;
+	};
+	FImportAssetObjectData ImportAssetObjectData;
 };
diff --git a/Engine/Plugins/Interchange/Runtime/Source/Import/Public/Scene/InterchangeSceneVariantSetsFactory.h b/Engine/Plugins/Interchange/Runtime/Source/Import/Public/Scene/InterchangeSceneVariantSetsFactory.h
--- a/Engine/Plugins/Interchange/Runtime/Source/Import/Public/Scene/InterchangeSceneVariantSetsFactory.h
+++ b/Engine/Plugins/Interchange/Runtime/Source/Import/Public/Scene/InterchangeSceneVariantSetsFactory.h
@@ -23,7 +23,7 @@ public:
 	// Interchange factory base interface begin
 
 	virtual UClass* GetFactoryClass() const override;
-	virtual UObject* ImportAssetObject_GameThread(const FImportAssetObjectParams& Arguments) override;
+	virtual UObject* BeginImportAssetObject_GameThread(const FImportAssetObjectParams& Arguments) override;
 	
 	virtual void SetupObject_GameThread(const FSetupObjectParams& Arguments) override;
 
diff --git a/Engine/Plugins/Interchange/Runtime/Source/Import/Public/Texture/InterchangeTextureFactory.h b/Engine/Plugins/Interchange/Runtime/Source/Import/Public/Texture/InterchangeTextureFactory.h
--- a/Engine/Plugins/Interchange/Runtime/Source/Import/Public/Texture/InterchangeTextureFactory.h
+++ b/Engine/Plugins/Interchange/Runtime/Source/Import/Public/Texture/InterchangeTextureFactory.h
@@ -56,7 +56,7 @@ public:
 
 	virtual UClass* GetFactoryClass() const override;
 	virtual EInterchangeFactoryAssetType GetFactoryAssetType() override { return EInterchangeFactoryAssetType::Textures; }
-	virtual UObject* ImportAssetObject_GameThread(const FImportAssetObjectParams& Arguments) override;
+	virtual UObject* BeginImportAssetObject_GameThread(const FImportAssetObjectParams& Arguments) override;
 	virtual UObject* ImportAssetObject_Async(const FImportAssetObjectParams& Arguments) override;
 	virtual void SetupObject_GameThread(const FSetupObjectParams& Arguments) override;
 	virtual bool GetSourceFilenames(const UObject* Object, TArray<FString>& OutSourceFilenames) const override;
diff --git a/Engine/Source/Developer/SkeletalMeshUtilitiesCommon/Private/LODUtilities.cpp b/Engine/Source/Developer/SkeletalMeshUtilitiesCommon/Private/LODUtilities.cpp
--- a/Engine/Source/Developer/SkeletalMeshUtilitiesCommon/Private/LODUtilities.cpp
+++ b/Engine/Source/Developer/SkeletalMeshUtilitiesCommon/Private/LODUtilities.cpp
@@ -32,7 +32,7 @@
 #include "HAL/ThreadSafeBool.h"
 #include "ImageCore.h"
 #include "ImageCoreUtils.h"
-
+#include "Rendering/SkeletalMeshLODImporterData.h"
 #include "ObjectTools.h"
 
 #include "ComponentReregisterContext.h"
@@ -55,6 +55,13 @@ IMPLEMENT_MODULE(FDefaultModuleImpl, SkeletalMeshUtilitiesCommon)
 
 DEFINE_LOG_CATEGORY_STATIC(LogLODUtilities, Log, All);
 
+static bool GInterchangeSkeletalMeshReorderMaterialSlots = true;
+static FAutoConsoleVariableRef CCvarInterchangeSkeletalMeshReorderMaterialSlots(
+	TEXT("Interchange.FeatureFlags.Import.SkeletalMesh.ReorderMaterialSlots"),
+	GInterchangeSkeletalMeshReorderMaterialSlots,
+	TEXT("Whether Re-importing a skeletalmesh should reorder the material slots."),
+	ECVF_Default);
+
 /**
 * Process and update the vertex Influences using the predefined wedges
 *
@@ -757,71 +764,21 @@ bool FLODUtilities::SetCustomLOD(USkeletalMesh* DestinationSkeletalMesh, USkelet
 		}
 	}
 
+	TArray<FName> ExistingOriginalPerSectionMaterialImportName;
+	bool bIsReImport = false;
 	//Restore the LOD section data in case this LOD was reimport and some material match
 	if (DestImportedResource->LODModels.IsValidIndex(LodIndex) && SourceImportedResource->LODModels.IsValidIndex(0))
 	{
-		const TArray<FSkelMeshSection>& ExistingSections = DestImportedResource->LODModels[LodIndex].Sections;
-		const FSkeletalMeshLODInfo& ExistingInfo = *(DestinationSkeletalMesh->GetLODInfo(LodIndex));
-
-		TArray<FSkelMeshSection>& ImportedSections = SourceImportedResource->LODModels[0].Sections;
-		const FSkeletalMeshLODInfo& ImportedInfo = *(SourceSkeletalMesh->GetLODInfo(0));
-
-		auto GetImportMaterialSlotName = [](const USkeletalMesh* SkelMesh, const FSkelMeshSection& Section, int32 SectionIndex, const FSkeletalMeshLODInfo& Info, int32& OutMaterialIndex)->FName
+		if (!DestinationSkeletalMesh->IsLODImportedDataEmpty(LodIndex))
 		{
-			const TArray<FSkeletalMaterial>& MeshMaterials = SkelMesh->GetMaterials();
-			check(MeshMaterials.Num() > 0);
-			OutMaterialIndex = Section.MaterialIndex;
-			if (Info.LODMaterialMap.IsValidIndex(SectionIndex) && Info.LODMaterialMap[SectionIndex] != INDEX_NONE)
-			{
-				OutMaterialIndex = Info.LODMaterialMap[SectionIndex];
-			}
-			FName ImportedMaterialSlotName = NAME_None;
-			if (MeshMaterials.IsValidIndex(OutMaterialIndex))
-			{
-				ImportedMaterialSlotName = MeshMaterials[OutMaterialIndex].ImportedMaterialSlotName;
-			}
-			else
+			FSkeletalMeshImportData DestinationLODImportData;
+			DestinationSkeletalMesh->LoadLODImportedData(LodIndex, DestinationLODImportData);
+			for (int32 SectionIndex = 0; SectionIndex < DestinationLODImportData.Materials.Num(); ++SectionIndex)
 			{
-				ImportedMaterialSlotName = MeshMaterials[0].ImportedMaterialSlotName;
-				OutMaterialIndex = 0;
-			}
-			return ImportedMaterialSlotName;
-		};
-
-		for (int32 ExistingSectionIndex = 0; ExistingSectionIndex < ExistingSections.Num(); ++ExistingSectionIndex)
-		{
-			const FSkelMeshSection& ExistingSection = ExistingSections[ExistingSectionIndex];
-			int32 ExistingMaterialIndex = 0;
-			FName ExistingImportedMaterialSlotName = GetImportMaterialSlotName(DestinationSkeletalMesh, ExistingSection, ExistingSectionIndex, ExistingInfo, ExistingMaterialIndex);
-
-			for (int32 ImportedSectionIndex = 0; ImportedSectionIndex < ImportedSections.Num(); ++ImportedSectionIndex)
-			{
-				FSkelMeshSection& ImportedSection = ImportedSections[ImportedSectionIndex];
-				int32 ImportedMaterialIndex = 0;
-				FName ImportedImportedMaterialSlotName = GetImportMaterialSlotName(SourceSkeletalMesh, ImportedSection, ImportedSectionIndex, ImportedInfo, ImportedMaterialIndex);
-				if (ExistingImportedMaterialSlotName != NAME_None)
-				{
-					if (ImportedImportedMaterialSlotName == ExistingImportedMaterialSlotName)
-					{
-						//Set the value and exit
-						ImportedSection.bCastShadow = ExistingSection.bCastShadow;
-						ImportedSection.bVisibleInRayTracing = ExistingSection.bVisibleInRayTracing;
-						ImportedSection.bRecomputeTangent = ExistingSection.bRecomputeTangent;
-						ImportedSection.RecomputeTangentsVertexMaskChannel = ExistingSection.RecomputeTangentsVertexMaskChannel;
-						break;
-					}
-				}
-				else if (SourceSkeletalMesh->GetMaterials()[ImportedMaterialIndex] == DestinationSkeletalMesh->GetMaterials()[ExistingMaterialIndex]) //Use material slot compare to match in case the name is none
-				{
-					//Set the value and exit
-					ImportedSection.bCastShadow = ExistingSection.bCastShadow;
-					ImportedSection.bVisibleInRayTracing = ExistingSection.bVisibleInRayTracing;
-					ImportedSection.bRecomputeTangent = ExistingSection.bRecomputeTangent;
-					ImportedSection.RecomputeTangentsVertexMaskChannel = ExistingSection.RecomputeTangentsVertexMaskChannel;
-					break;
-				}
+				ExistingOriginalPerSectionMaterialImportName.Add(FName(*DestinationLODImportData.Materials[SectionIndex].MaterialImportName));
 			}
 		}
+		bIsReImport = true;
 	}
 
 	// If we want to add this as a new LOD to this mesh - add to LODModels/LODInfo array.
@@ -834,6 +791,22 @@ bool FLODUtilities::SetCustomLOD(USkeletalMesh* DestinationSkeletalMesh, USkelet
 		check(DestinationSkeletalMesh->GetLODNum() == DestImportedResource->LODModels.Num());
 	}
 
+	const int32 SourceLODIndex = 0;
+	if (!SourceSkeletalMesh->IsLODImportedDataEmpty(SourceLODIndex))
+	{
+		// Fix up the imported data bone indexes
+		FSkeletalMeshImportData SourceLODImportData;
+		SourceSkeletalMesh->LoadLODImportedData(SourceLODIndex, SourceLODImportData);
+		FLODUtilities::FSkeletalMeshMatchImportedMaterialsParameters Parameters;
+		Parameters.bIsReImport = bIsReImport;
+		Parameters.LodIndex = LodIndex;
+		Parameters.SkeletalMesh = DestinationSkeletalMesh;
+		Parameters.ImportedMaterials = &SourceLODImportData.Materials;
+		Parameters.ExistingOriginalPerSectionMaterialImportName = &ExistingOriginalPerSectionMaterialImportName;
+		Parameters.CustomImportedLODModel = &NewLODModel;
+		FLODUtilities::MatchImportedMaterials(Parameters);
+	}
+
 	// Set up LODMaterialMap to number of materials in new mesh.
 	FSkeletalMeshLODInfo& LODInfo = *(DestinationSkeletalMesh->GetLODInfo(LodIndex));
 
@@ -844,34 +817,6 @@ bool FLODUtilities::SetCustomLOD(USkeletalMesh* DestinationSkeletalMesh, USkelet
 		LODInfo.BuildSettings = ImportedLODInfo.BuildSettings;
 	}
 
-	TArray<FSkeletalMaterial>& BaseMaterials = DestinationSkeletalMesh->GetMaterials();
-	LODInfo.LODMaterialMap.Empty();
-	// Now set up the material mapping array.
-	for (int32 SectionIndex = 0; SectionIndex < NewLODModel.Sections.Num(); SectionIndex++)
-	{
-		int32 MatIdx = NewLODModel.Sections[SectionIndex].MaterialIndex;
-		// Try and find the auto-assigned material in the array.
-		int32 LODMatIndex = INDEX_NONE;
-		//First try to match by name
-		for (int32 BaseMaterialIndex = 0; BaseMaterialIndex < BaseMaterials.Num(); ++BaseMaterialIndex)
-		{
-			const FSkeletalMaterial& SkeletalMaterial = BaseMaterials[BaseMaterialIndex];
-			if (SkeletalMaterial.ImportedMaterialSlotName != NAME_None && SkeletalMaterial.ImportedMaterialSlotName == SourceSkeletalMesh->GetMaterials()[MatIdx].ImportedMaterialSlotName)
-			{
-				LODMatIndex = BaseMaterialIndex;
-				break;
-			}
-		}
-
-		// If we dont have a match, add a new entry to the material list.
-		if (LODMatIndex == INDEX_NONE)
-		{
-			LODMatIndex = BaseMaterials.Add(SourceSkeletalMesh->GetMaterials()[MatIdx]);
-		}
-
-		LODInfo.LODMaterialMap.Add(LODMatIndex);
-	}
-
 	// Release all resources before replacing the model
 	DestinationSkeletalMesh->PreEditChange(NULL);
 
@@ -3572,6 +3517,396 @@ void FLODUtilities::AdjustImportDataFaceMaterialIndex(const TArray<FSkeletalMate
 		}
 	}
 }
+
+void FLODUtilities::MatchImportedMaterials(FLODUtilities::FSkeletalMeshMatchImportedMaterialsParameters& Parameters)
+{
+	//Make sure we have valid parameters
+	if (!ensure(Parameters.SkeletalMesh))
+	{
+		UE_LOG(LogLODUtilities, Warning, TEXT("FLODUtilities::MatchImportedMaterials: Bad parameters, SkeletalMesh is null"));
+		return;
+	}
+	if (!ensure(Parameters.ImportedMaterials))
+	{
+		UE_ASSET_LOG(LogLODUtilities, Warning, Parameters.SkeletalMesh, TEXT("FLODUtilities::MatchImportedMaterials: Bad parameters, ImportedMaterials is null"));
+		return;
+	}
+	if (Parameters.bIsReImport)
+	{
+		if (!ensure(Parameters.ExistingOriginalPerSectionMaterialImportName))
+		{
+			UE_ASSET_LOG(LogLODUtilities, Warning, Parameters.SkeletalMesh, TEXT("FLODUtilities::MatchImportedMaterials: Bad parameters, ExistingOriginalPerSectionMaterialImportName is null"));
+			return;
+		}
+	}
+	FSkeletalMeshLODInfo* LodInfo = Parameters.SkeletalMesh->GetLODInfo(Parameters.LodIndex);
+	if (!LodInfo)
+	{
+		UE_ASSET_LOG(LogLODUtilities, Warning, Parameters.SkeletalMesh, TEXT("FLODUtilities::MatchImportedMaterials: LodIndex %d, LodInfo is null"), Parameters.LodIndex);
+		return;
+	}
+	//Grab some parameters pointer into references
+	const TArray<SkeletalMeshImportData::FMaterial>& ImportedMaterials = *Parameters.ImportedMaterials;
+
+	TArray<FSkeletalMaterial>& Materials = Parameters.SkeletalMesh->GetMaterials();
+
+	TMap<FName, int32> LODMaterialMapRedirection;
+	//If we reimport we have to keep the existing user section info data. We use the material name to match the existing
+	if (Parameters.bIsReImport)
+	{
+		const TArray<FName>& ExistingOriginalPerSectionMaterialImportName = *Parameters.ExistingOriginalPerSectionMaterialImportName;
+
+		auto GetImportMaterialSlotName = [&Parameters, &LodInfo](const FSkelMeshSection& Section
+			, int32 SectionIndex
+			, int32& OutMaterialIndex
+			, const TArray<FName>& PerSectionMaterialImportName)->FName
+		{
+			const TArray<FSkeletalMaterial>& MeshMaterials = Parameters.SkeletalMesh->GetMaterials();
+			check(MeshMaterials.Num() > 0);
+			if (PerSectionMaterialImportName.IsValidIndex(SectionIndex))
+			{
+				for (int32 MaterialIndex = 0; MaterialIndex < MeshMaterials.Num(); ++MaterialIndex)
+				{
+					const FSkeletalMaterial& Material = MeshMaterials[MaterialIndex];
+					if (PerSectionMaterialImportName[SectionIndex] == Material.ImportedMaterialSlotName)
+					{
+						OutMaterialIndex = MaterialIndex;
+						break;
+					}
+				}
+				return PerSectionMaterialImportName[SectionIndex];
+			}
+
+			OutMaterialIndex = Section.MaterialIndex;
+			if (LodInfo->LODMaterialMap.IsValidIndex(SectionIndex) && LodInfo->LODMaterialMap[SectionIndex] != INDEX_NONE)
+			{
+				OutMaterialIndex = LodInfo->LODMaterialMap[SectionIndex];
+			}
+			FName ImportedMaterialSlotName = NAME_None;
+			if (MeshMaterials.IsValidIndex(OutMaterialIndex))
+			{
+				ImportedMaterialSlotName = MeshMaterials[OutMaterialIndex].ImportedMaterialSlotName;
+			}
+			else
+			{
+				ImportedMaterialSlotName = MeshMaterials[0].ImportedMaterialSlotName;
+				OutMaterialIndex = 0;
+			}
+			return ImportedMaterialSlotName;
+		};
+
+		TMap<FName, FSkelMeshSourceSectionUserData> NewUserSectionsDataMap;
+		FSkeletalMeshLODModel& ExistLODModel = Parameters.SkeletalMesh->GetImportedModel()->LODModels[Parameters.LodIndex];
+		for (int32 NewSectionIndex = 0; NewSectionIndex < ImportedMaterials.Num(); NewSectionIndex++)
+		{
+			FName ImportedMaterialName = *(ImportedMaterials[NewSectionIndex].MaterialImportName);
+			FSkelMeshSourceSectionUserData& SourceSectionUserData = NewUserSectionsDataMap.FindOrAdd(ImportedMaterialName);
+
+			for (int32 ExistSectionIndex = 0; ExistSectionIndex < ExistLODModel.Sections.Num(); ExistSectionIndex++)
+			{
+				if (ExistLODModel.Sections[ExistSectionIndex].ChunkedParentSectionIndex != INDEX_NONE)
+				{
+					continue;
+				}
+
+				const FSkelMeshSection& ExistingSection = ExistLODModel.Sections[ExistSectionIndex];
+				//Skip chunked section
+				if (ExistingSection.ChunkedParentSectionIndex != INDEX_NONE)
+					continue;
+				int32 ExistingMaterialIndex = 0;
+				FName ExistingImportedMaterialSlotName = GetImportMaterialSlotName(ExistingSection
+					, ExistingSection.OriginalDataSectionIndex
+					, ExistingMaterialIndex
+					, ExistingOriginalPerSectionMaterialImportName);
+
+				TArray<FName> EmptyArray;
+				int32 ExistingCurrentMaterialIndex = 0;
+				FName ExistingCurrentImportedMaterialSlotName = GetImportMaterialSlotName(ExistingSection
+					, ExistingSection.OriginalDataSectionIndex
+					, ExistingCurrentMaterialIndex
+					, EmptyArray);
+
+				auto CopySectionData = [&ExistingSection, &SourceSectionUserData]()
+				{
+					//Set the user section data to reflect the existing settings
+					SourceSectionUserData.bCastShadow = ExistingSection.bCastShadow;
+					SourceSectionUserData.bVisibleInRayTracing = ExistingSection.bVisibleInRayTracing;
+					SourceSectionUserData.bRecomputeTangent = ExistingSection.bRecomputeTangent;
+					SourceSectionUserData.RecomputeTangentsVertexMaskChannel = ExistingSection.RecomputeTangentsVertexMaskChannel;
+				};
+
+				if (ExistingImportedMaterialSlotName != NAME_None)
+				{
+					if (ImportedMaterialName == ExistingImportedMaterialSlotName)
+					{
+						CopySectionData();
+						//If user has edit the material slot pointer, store the data so we can modify the LODMaterialMap
+						if (ExistingCurrentImportedMaterialSlotName != ExistingImportedMaterialSlotName)
+						{
+							LODMaterialMapRedirection.FindOrAdd(ImportedMaterialName) = ExistingCurrentMaterialIndex;
+						}
+						break;
+					}
+				}
+				else if (ImportedMaterials[NewSectionIndex].Material.Get() == Parameters.SkeletalMesh->GetMaterials()[ExistingCurrentMaterialIndex].MaterialInterface.Get()) //Use material slot compare to match in case the name is none
+				{
+					CopySectionData();
+					break;
+				}
+			}
+		}
+
+		for (int32 NewSectionIndex = 0; NewSectionIndex < ImportedMaterials.Num(); NewSectionIndex++)
+		{
+			FName ImportedMaterialName = *(ImportedMaterials[NewSectionIndex].MaterialImportName);
+			if (!NewUserSectionsDataMap.Contains(ImportedMaterialName))
+			{
+				NewUserSectionsDataMap.Add(ImportedMaterialName, FSkelMeshSourceSectionUserData());
+			}
+		}
+
+		if (Parameters.LodIndex == 0)
+		{
+			//Since LOD 0 data will be re-arrange to follow the Material list order, sort the map accordingly.
+			NewUserSectionsDataMap.KeySort([Materials](const FName& A, const FName& B)
+				{
+					int32 MatAIndex = INDEX_NONE;
+					int32 MatBIndex = INDEX_NONE;
+					for (int32 MaterialIndex = 0; MaterialIndex < Materials.Num(); ++MaterialIndex)
+					{
+						if (Materials[MaterialIndex].ImportedMaterialSlotName == A)
+						{
+							MatAIndex = MaterialIndex;
+						}
+						if (Materials[MaterialIndex].ImportedMaterialSlotName == B)
+						{
+							MatBIndex = MaterialIndex;
+						}
+					}
+					if (MatBIndex == INDEX_NONE)
+					{
+						return true;
+					}
+					if (MatAIndex == INDEX_NONE)
+					{
+						return false;
+					}
+					return MatAIndex < MatBIndex;
+				});
+		}
+		auto ApplyUserSectionLODModel = [&NewUserSectionsDataMap](FSkeletalMeshLODModel& ToApplyUserSectionLODModel, bool bSynchronizeUserSectionData)
+		{
+			ToApplyUserSectionLODModel.UserSectionsData.Reset();
+			int32 RemapSectionIndex = 0;
+			for (TPair<FName, FSkelMeshSourceSectionUserData> UserSectionsData : NewUserSectionsDataMap)
+			{
+				ToApplyUserSectionLODModel.UserSectionsData.Add(RemapSectionIndex++, UserSectionsData.Value);
+			}
+			if (bSynchronizeUserSectionData)
+			{
+				ToApplyUserSectionLODModel.SyncronizeUserSectionsDataArray(true);
+			}
+		};
+
+		//If the caller provide a CustomImportedLODModel we must apply the user section data to the provided lod model.
+		if (Parameters.CustomImportedLODModel)
+		{
+			constexpr bool bSynchronizeUserSectionData = true;
+			ApplyUserSectionLODModel(*Parameters.CustomImportedLODModel, bSynchronizeUserSectionData);
+		}
+		else
+		{
+			constexpr bool bSynchronizeUserSectionData = false;
+			ApplyUserSectionLODModel(ExistLODModel, bSynchronizeUserSectionData);
+		}
+		
+	}
+
+	LodInfo->LODMaterialMap.Empty();
+	// Now set up the material mapping array.
+	if (Parameters.LodIndex != 0)
+	{
+		for (int32 ImportedMaterialIndex = 0; ImportedMaterialIndex < ImportedMaterials.Num(); ImportedMaterialIndex++)
+		{
+			int32 LODMatIndex = INDEX_NONE;
+			FName ImportedMaterialName = *(ImportedMaterials[ImportedMaterialIndex].MaterialImportName);
+
+			//If the current lod section material slot was edited, we must keep the edition.
+			if (int32* MaterialRedirection = LODMaterialMapRedirection.Find(ImportedMaterialName))
+			{
+				LODMatIndex = *MaterialRedirection;
+			}
+			else
+			{
+				//Match by name
+				for (int32 MaterialIndex = 0; MaterialIndex < Materials.Num(); ++MaterialIndex)
+				{
+					const FSkeletalMaterial& SkeletalMaterial = Materials[MaterialIndex];
+					if (SkeletalMaterial.ImportedMaterialSlotName != NAME_None && SkeletalMaterial.ImportedMaterialSlotName == ImportedMaterialName)
+					{
+						LODMatIndex = MaterialIndex;
+						break;
+					}
+				}
+
+				// If we don't have a match, add a new entry to the material list.
+				if (LODMatIndex == INDEX_NONE)
+				{
+					LODMatIndex = Materials.Add(FSkeletalMaterial(ImportedMaterials[ImportedMaterialIndex].Material.Get(), true, false, ImportedMaterialName, ImportedMaterialName));
+				}
+			}
+			LodInfo->LODMaterialMap.Add(LODMatIndex);
+		}
+	}
+	else if (LODMaterialMapRedirection.Num() > 0)
+	{
+		for (int32 ExistMaterialIndex = 0; ExistMaterialIndex < Materials.Num(); ExistMaterialIndex++)
+		{
+			int32 LODMatIndex = INDEX_NONE;
+			FName ExistMaterialName = Materials[ExistMaterialIndex].ImportedMaterialSlotName;
+			//If the current lod section material slot was edited, we must keep the edition.
+			if (int32* MaterialRedirection = LODMaterialMapRedirection.Find(ExistMaterialName))
+			{
+				LODMatIndex = *MaterialRedirection;
+			}
+			LodInfo->LODMaterialMap.Add(LODMatIndex);
+		}
+	}
+
+	if (Parameters.bIsReImport && Parameters.LodIndex == 0 && GInterchangeSkeletalMeshReorderMaterialSlots)
+	{
+		ReorderMaterialSlotToBaseLod(Parameters.SkeletalMesh);
+	}
+}
+
+void FLODUtilities::ReorderMaterialSlotToBaseLod(USkeletalMesh* SkeletalMesh)
+{
+	//Make sure we have valid parameters
+	if (!ensure(SkeletalMesh))
+	{
+		UE_LOG(LogLODUtilities, Warning, TEXT("FLODUtilities::ReorderMaterialSlotToBaseLod: Bad parameters, SkeletalMesh is null"));
+		return;
+	}
+	if (SkeletalMesh->IsLODImportedDataEmpty(0))
+	{
+		UE_LOG(LogLODUtilities, Warning, TEXT("FLODUtilities::ReorderMaterialSlotToBaseLod: Skeletal mesh invalid import data for LOD 0"));
+		return;
+	}
+	FSkeletalMeshImportData BaseLodImportData;
+	SkeletalMesh->LoadLODImportedData(0, BaseLodImportData);
+	const TArray<SkeletalMeshImportData::FMaterial>& ImportedMaterials = BaseLodImportData.Materials;
+	TArray<FSkeletalMaterial>& Materials = SkeletalMesh->GetMaterials();
+
+	TArray<int32> MaterialSlotRemap;
+	MaterialSlotRemap.AddUninitialized(Materials.Num());
+	TBitArray<> AvailableIndexes;
+	AvailableIndexes.Init(true, Materials.Num());
+	//Find remap index for LOD 0 matching materials
+	for (int32 MaterialIndex = 0; MaterialIndex < Materials.Num(); ++MaterialIndex)
+	{
+		FName MaterialSlotName = Materials[MaterialIndex].ImportedMaterialSlotName;
+		MaterialSlotRemap[MaterialIndex] = INDEX_NONE;
+		for (int32 ImportedMaterialIndex = 0; ImportedMaterialIndex < ImportedMaterials.Num(); ++ImportedMaterialIndex)
+		{
+			FName ImportedMaterialSlotName = FName(*ImportedMaterials[ImportedMaterialIndex].MaterialImportName);
+			if (MaterialSlotName == ImportedMaterialSlotName)
+			{
+				MaterialSlotRemap[MaterialIndex] = ImportedMaterialIndex;
+				AvailableIndexes[MaterialIndex] = false;
+				break;
+			}
+		}
+	}
+	//Find remap index for any extra custom LOD materials
+	for (int32 MaterialIndex = 0; MaterialIndex < Materials.Num(); ++MaterialIndex)
+	{
+		if (MaterialSlotRemap[MaterialIndex] == INDEX_NONE)
+		{
+			//Give the first available material index
+			for (int32 AvailableIndex = 0; AvailableIndex < AvailableIndexes.Num(); ++AvailableIndex)
+			{
+				if (AvailableIndexes[AvailableIndex])
+				{
+					MaterialSlotRemap[MaterialIndex] = AvailableIndex;
+					AvailableIndexes[AvailableIndex] = false;
+					break;
+				}
+			}
+			//We should always be able to rematch a material
+			ensure(MaterialSlotRemap[MaterialIndex] != INDEX_NONE);
+		}
+	}
+
+	//Build the re-ordered maps
+	TArray<FSkeletalMaterial> ReorderMaterials;
+	ReorderMaterials.Reserve(Materials.Num());
+	for (int32 MaterialIndex = 0; MaterialIndex < Materials.Num(); ++MaterialIndex)
+	{
+		for (int32 RemapIndex = 0; RemapIndex < MaterialSlotRemap.Num(); ++RemapIndex)
+		{
+			if (MaterialSlotRemap[RemapIndex] == MaterialIndex)
+			{
+				ReorderMaterials.Add(Materials[RemapIndex]);
+			}
+		}
+	}
+
+	//Change the material slot array
+	SkeletalMesh->SetMaterials(ReorderMaterials);
+	
+	//We now need to adjust all LODs data to fit the re-order
+	for (int32 LodIndex = 0; LodIndex < SkeletalMesh->GetLODNum(); ++LodIndex)
+	{
+		if (ensure(SkeletalMesh->GetImportedModel()->LODModels.IsValidIndex(LodIndex)))
+		{
+			FSkeletalMeshLODModel& LODModel = SkeletalMesh->GetImportedModel()->LODModels[LodIndex];
+			//The skeletal mesh build process have a different behavior for LOD 0. The build reorder the import data sections in the material slot order.
+			//Because of this we need to remap the UserSectionData key
+			if (LodIndex == 0)
+			{
+				TMap<int32, FSkelMeshSourceSectionUserData> RemapSectionsDataMap;
+				for (TPair<int32, FSkelMeshSourceSectionUserData>& UserSectionDataPair : LODModel.UserSectionsData)
+				{
+					RemapSectionsDataMap.Add(MaterialSlotRemap[UserSectionDataPair.Key], UserSectionDataPair.Value);
+				}
+				//Sort the remap section so its in the proper order. Optional but easier to debug ordered data
+				RemapSectionsDataMap.KeySort([Materials](const int32& A, const int32& B)
+					{
+						return A < B;
+					});
+				LODModel.UserSectionsData = RemapSectionsDataMap;
+			}
+
+			//Remap the built sections material index, note that the sections should be rebuild after changing the material slot order. so not a critical step.
+			for (int32 SectionIndex = 0; SectionIndex < LODModel.Sections.Num(); ++SectionIndex)
+			{
+				FSkelMeshSection& Section = LODModel.Sections[SectionIndex];
+				Section.MaterialIndex = MaterialSlotRemap[Section.MaterialIndex];
+			}
+		}
+		else
+		{
+			UE_LOG(LogLODUtilities, Warning, TEXT("FLODUtilities::ReorderMaterialSlotToBaseLod: Skeletal mesh invalid LODModel %d"), LodIndex);
+		}
+		
+		//Remap the LODMaterialMap
+		if (FSkeletalMeshLODInfo* LodInfo = SkeletalMesh->GetLODInfo(LodIndex))
+		{
+			for (int32& MaterialMap : LodInfo->LODMaterialMap)
+			{
+				if (MaterialSlotRemap.IsValidIndex(MaterialMap))
+				{
+					MaterialMap = MaterialSlotRemap[MaterialMap];
+				}
+			}
+		}
+		else
+		{
+			UE_LOG(LogLODUtilities, Warning, TEXT("FLODUtilities::ReorderMaterialSlotToBaseLod: Skeletal mesh invalid LODInfo %d"), LodIndex);
+		}
+	}
+}
+
 namespace TriangleStripHelper
 {
 	struct FTriangle2D
diff --git a/Engine/Source/Developer/SkeletalMeshUtilitiesCommon/Public/LODUtilities.h b/Engine/Source/Developer/SkeletalMeshUtilitiesCommon/Public/LODUtilities.h
--- a/Engine/Source/Developer/SkeletalMeshUtilitiesCommon/Public/LODUtilities.h
+++ b/Engine/Source/Developer/SkeletalMeshUtilitiesCommon/Public/LODUtilities.h
@@ -233,6 +233,37 @@ public:
 	 */
 	static void AdjustImportDataFaceMaterialIndex(const TArray<FSkeletalMaterial>& Materials, TArray<SkeletalMeshImportData::FMaterial>& RawMeshMaterials, TArray<SkeletalMeshImportData::FMeshFace>& LODFaces, int32 LODIndex);
 
+	/**
+	 * Structure to pass all the needed parameters to do match the material when importing a skeletal mesh LOD.
+	 * @param SkeletalMesh - The skeletalmesh that get a LOD re/import
+	 * @param LodIndex - The skeletalmesh LOD index that get re/import
+	 * @param bIsReImport - whether its a skeletal mesh LOD re-import
+	 * @param ImportedMaterials - The imported material list
+	 * @param ExistingOriginalPerSectionMaterialImportName - The previously LOD imported material list (cannot be null if bIsReImport is true)
+	 * @param CustomImportedLODModel - When importing custom LOD we want to pass the LODModel to synchronize the sections
+	 */
+	struct FSkeletalMeshMatchImportedMaterialsParameters
+	{
+		USkeletalMesh* SkeletalMesh = nullptr;
+		int32 LodIndex = 0;
+		bool bIsReImport = false;
+		const TArray<SkeletalMeshImportData::FMaterial>* ImportedMaterials = nullptr;
+		const TArray<FName>* ExistingOriginalPerSectionMaterialImportName = nullptr;
+		FSkeletalMeshLODModel* CustomImportedLODModel = nullptr;
+	};
+
+	/**
+	 * When any skeletalmesh LOD get imported or re-imported, we want to have common code to set all materials data.
+	 * In case we re-import a LOD this code will match the sections using imported material name to be able to keep any section data.
+	 * The material slot array is not re-order even if we re-import the base LOD.
+	 */
+	static void MatchImportedMaterials(FSkeletalMeshMatchImportedMaterialsParameters& Parameters);
+
+	/**
+	 * Reorder the material slot array to follow the base LOD section order. It will readjust all LOD section material index and LODMaterialMap.
+	 */
+	static void ReorderMaterialSlotToBaseLod(USkeletalMesh* SkeletalMesh);
+
 	/**
 	 * This function will strip all triangle in the specified LOD that don't have any UV area pointing on a black pixel in the TextureMask.
 	 * We use the UVChannel 0 to find the pixels in the texture.
diff --git a/Engine/Source/Runtime/Engine/Private/StaticMesh.cpp b/Engine/Source/Runtime/Engine/Private/StaticMesh.cpp
--- a/Engine/Source/Runtime/Engine/Private/StaticMesh.cpp
+++ b/Engine/Source/Runtime/Engine/Private/StaticMesh.cpp
@@ -64,6 +64,7 @@
 #endif // #if WITH_EDITOR
 
 #include "Engine/StaticMeshSocket.h"
+#include "MaterialDomain.h"
 #include "EditorFramework/AssetImportData.h"
 #include "AI/Navigation/NavCollisionBase.h"
 #include "AI/NavigationSystemBase.h"
@@ -80,6 +81,10 @@
 #include "Widgets/Notifications/SNotificationList.h"
 #endif
 
+#if WITH_EDITORONLY_DATA
+#include "Materials/Material.h"
+#endif
+
 #define LOCTEXT_NAMESPACE "StaticMesh"
 DEFINE_LOG_CATEGORY(LogStaticMesh);	
 
@@ -4203,6 +4208,8 @@ bool UStaticMesh::SetCustomLOD(const UStaticMesh* SourceStaticMesh, int32 Destin
 		return false;
 	}
 
+	const bool bIsReimport = GetNumSourceModels() > DestinationLodIndex;
+
 	if(DestinationLodIndex >= GetNumSourceModels())
 	{
 		// Add one LOD 
@@ -4228,7 +4235,7 @@ bool UStaticMesh::SetCustomLOD(const UStaticMesh* SourceStaticMesh, int32 Destin
 	}
 
 	TArray<FStaticMaterial>& DestinationMaterials = GetStaticMaterials();
-
+	TMap<FName, FMeshSectionInfo> ExistingSectionInfos;
 	FMeshDescription* DestinationMeshDescription = GetMeshDescription(DestinationLodIndex);
 	if (DestinationMeshDescription == nullptr)
 	{
@@ -4243,58 +4250,81 @@ bool UStaticMesh::SetCustomLOD(const UStaticMesh* SourceStaticMesh, int32 Destin
 	}
 	else
 	{
+		ensure(bIsReimport);
+		FStaticMeshConstAttributes ExistingAttributes(*DestinationMeshDescription);
+		TPolygonGroupAttributesConstRef<FName> ExistingMaterialSlotNames = ExistingAttributes.GetPolygonGroupMaterialSlotNames();
+		int32 SectionIndex = 0;
+		for (FPolygonGroupID PolygonGroupID : DestinationMeshDescription->PolygonGroups().GetElementIDs())
+		{
+			FMeshSectionInfo& ExistingInfo = ExistingSectionInfos.FindOrAdd(ExistingMaterialSlotNames[PolygonGroupID]);
+			int32 MaterialSlotIndex = GetMaterialIndexFromImportedMaterialSlotName(ExistingMaterialSlotNames[PolygonGroupID]);
+			if (GetSectionInfoMap().IsValidSection(DestinationLodIndex, SectionIndex))
+			{
+				ExistingInfo = GetSectionInfoMap().Get(DestinationLodIndex, SectionIndex);
+			}
+			else
+			{
+				ExistingInfo.MaterialIndex = MaterialSlotIndex;
+			}
+			if (!DestinationMaterials.IsValidIndex(ExistingInfo.MaterialIndex))
+			{
+				//There was an invalid material index in the existing mesh section info
+				//Assign the raw material slot that match with the name if valid, otherwise set it at 0
+				ExistingInfo.MaterialIndex = MaterialSlotIndex != INDEX_NONE ? MaterialSlotIndex : 0;
+			}
+			SectionIndex++;
+		}
 		// clear out the old mesh data
 		DestinationMeshDescription->Empty();
 	}
 
 	//Make sure all materials use by the new LOD is pointing on a valid static material
+	int32 SectionIndex = 0;
 	for(const TPair<int32, FName>& SourceImportedMaterialNamePair : SourceImportedMaterialNameUsed)
 	{
-		bool bFoundMatch = false;
 		FName NameSearch = SourceImportedMaterialNamePair.Value;
-			
-		for (int32 MaterialIndex = 0; MaterialIndex < DestinationMaterials.Num(); ++MaterialIndex)
+		int32 MaterialSlotIndex = GetMaterialIndexFromImportedMaterialSlotName(NameSearch);
+		if (!ExistingSectionInfos.Contains(NameSearch))
 		{
-			const FStaticMaterial& StaticMaterial = DestinationMaterials[MaterialIndex];
-			if (NameSearch == StaticMaterial.ImportedMaterialSlotName)
+			if (MaterialSlotIndex == INDEX_NONE)
 			{
-				bFoundMatch = true;
-				break;
+				//Add the missing material slot
+				MaterialSlotIndex = DestinationMaterials.Add(FStaticMaterial(UMaterial::GetDefaultMaterial(MD_Surface), NameSearch, NameSearch));
 			}
+			FMeshSectionInfo NewInfo;
+			NewInfo.MaterialIndex = MaterialSlotIndex;
+			ExistingSectionInfos.Add(NameSearch, NewInfo);
 		}
-
-		if (!bFoundMatch)
-		{
-			//Add the missing material slot
-			FStaticMaterial StaticMaterial(nullptr, NameSearch, NameSearch);
-		}
+#if WITH_EDITOR
+		FMeshSectionInfo Info = ExistingSectionInfos.FindChecked(NameSearch);
+		GetSectionInfoMap().Remove(DestinationLodIndex, SectionIndex);
+		GetSectionInfoMap().Set(DestinationLodIndex, SectionIndex, Info);
+#endif //WITH_EDITOR
+		SectionIndex++;
 	}
 
 	//Copy the mesh description of the source into the destination
 	*DestinationMeshDescription = *SourceMeshDescription;
 	
-
-	if(IsInGameThread())
+	UStaticMesh* ThisMesh = this;
+	auto FinalizeSetCustomLODGameThread = [ThisMesh, DestinationLodIndex]()
 	{
-		PostEditChange();
+		check(IsInGameThread());
 		//Commit the mesh description to update the ddc key
 		FCommitMeshDescriptionParams CommitMeshDescriptionParams;
 		CommitMeshDescriptionParams.bMarkPackageDirty = true;
 		CommitMeshDescriptionParams.bUseHashAsGuid = false;
-		CommitMeshDescription(DestinationLodIndex, CommitMeshDescriptionParams);
+		ThisMesh->CommitMeshDescription(DestinationLodIndex, CommitMeshDescriptionParams);
+		ThisMesh->PostEditChange();
+	};
+
+	if(IsInGameThread())
+	{
+		FinalizeSetCustomLODGameThread();
 	}
 	else
 	{
-		UStaticMesh* ThisMesh = this;
-		Async(EAsyncExecution::TaskGraphMainThread, [ThisMesh, DestinationLodIndex]()
-		{
-			ThisMesh->PostEditChange();
-			//Commit the mesh description to update the ddc key
-			FCommitMeshDescriptionParams CommitMeshDescriptionParams;
-			CommitMeshDescriptionParams.bMarkPackageDirty = true;
-			CommitMeshDescriptionParams.bUseHashAsGuid = false;
-			ThisMesh->CommitMeshDescription(DestinationLodIndex, CommitMeshDescriptionParams);
-		});
+		Async(EAsyncExecution::TaskGraphMainThread, MoveTemp(FinalizeSetCustomLODGameThread));
 	}
 
 	if (IsSourceModelValid(DestinationLodIndex))
diff --git a/Engine/Source/Runtime/Interchange/Core/Public/InterchangeFactoryBase.h b/Engine/Source/Runtime/Interchange/Core/Public/InterchangeFactoryBase.h
--- a/Engine/Source/Runtime/Interchange/Core/Public/InterchangeFactoryBase.h
+++ b/Engine/Source/Runtime/Interchange/Core/Public/InterchangeFactoryBase.h
@@ -30,13 +30,22 @@ enum class EInterchangeFactoryAssetType : uint8
 };
 
 /**
- * UObject asset factory implementation:
- * 1. ImportAssetObject_GameThread - Create the asset UObject, and you can also import source data and setup properties synchronously.
- * 2. ImportAssetObject_Async - Import source data and setup properties asynchronously.
- * 3. SetupObject_GameThread - Do any UObject setup required before the build (before PostEditChange), the UObject dependencies should exist and have all the source data and properties imported.
- * 4. FinalizeObject_GameThread - Do any final UObject setup after the build (after PostEditChange)
+ * Asset factory implementation:
+ * 
+ * The first 3 steps use the interchange factory node to import/re-import the UObject
+ * 
+ * 1. BeginImportAssetObject_GameThread - Create the asset UPackage and the asset UObject. You can also import source data (retrieve payloads) and setup properties on the game thread.
+ * 2. ImportAssetObject_Async - Import source data (retrieve payloads) and setup properties asynchronously on any thread.
+ * 3. EndImportAssetObject_GameThread - Anything you need to do on the game thread to finalize the import source data and setup properties. (Conflict resolution that need UI)
+ * 
+ * The last two steps are helpful to change the imported/re-imported UObject before and after the asset is built.
+ * 
+ * 4. SetupObject_GameThread - Do any UObject setup required before the build (before PostEditChange), the UObject dependencies should exist and have all the source data and properties imported.
+ * 5. FinalizeObject_GameThread - Do any final UObject setup after the build (after PostEditChange)
+ * 
+ * 
+ * Scene factory implementation
  * 
- * Actor factory implementation
  * 1. ImportSceneObject_GameThread - Create an actor in a level.
  */
 
@@ -90,49 +99,62 @@ public:
 		UObject* ReimportObject = nullptr;
 	};
 
-	UE_DEPRECATED(5.2, "This function is replaced by ImportAssetObject_GameThread.")
-	virtual UObject* CreateEmptyAsset(const FImportAssetObjectParams& Arguments)
+	UE_DEPRECATED(5.3, "This function is replaced by BeginImportAssetObject_GameThread.")
+	virtual UObject* ImportAssetObject_GameThread(const FImportAssetObjectParams& Arguments)
 	{
 		return nullptr;
 	}
-	
+
 	/**
-	 * Create the asset UObject. You can optionally import source data (payload data) and configure the properties synchronously.
+	 * Override it to import/re-import source data and configure the properties synchronously.
+	 * Create the asset package on the game thread since its not thread safe.
 	 *
 	 * @param Arguments - The structure containing all necessary arguments to import the UObject and the package, see the structure definition for the documentation.
 	 * @return the imported UObject or nullptr if there is an error.
 	 *
-	 * @Note Mandatory to override this function to create the asset UObject. Not mandatory for level actor, use CreateSceneObject instead.
+	 * @Note Mandatory to override this function to create the asset UObject package. Not mandatory for level actor, use CreateSceneObject instead.
 	 */
-	virtual UObject* ImportAssetObject_GameThread(const FImportAssetObjectParams& Arguments)
+	virtual UObject* BeginImportAssetObject_GameThread(const FImportAssetObjectParams& Arguments)
 	{
 		PRAGMA_DISABLE_DEPRECATION_WARNINGS
-		return CreateEmptyAsset(Arguments);
+		return ImportAssetObject_GameThread(Arguments);
 		PRAGMA_ENABLE_DEPRECATION_WARNINGS
 	}
 
-	UE_DEPRECATED(5.2, "This function is replaced by ImportAssetObject_Async.")
-	virtual UObject* CreateAsset(const FImportAssetObjectParams& Arguments)
+	/**
+	 * Override it to import/re-import the UObject source data and configure the properties asynchronously. Do not create the asset UObject asynchronously.
+	 * Helpful to get all the payload in parallel or do any long import task
+	 *
+	 * @param Arguments - The structure containing all necessary data to imported the UObject, see the structure definition for the documentation.
+	 * @return the imported UObject or nullptr if there is an error.
+	 */
+	virtual UObject* ImportAssetObject_Async(const FImportAssetObjectParams& Arguments)
 	{
 		//By default simply return the UObject created by ImportAssetObject_GameThread that was store in the asset node.
 		FSoftObjectPath ReferenceObject;
-		if (Arguments.AssetNode->GetCustomReferenceObject(ReferenceObject))
+		if (Arguments.AssetNode && Arguments.AssetNode->GetCustomReferenceObject(ReferenceObject))
 		{
 			return ReferenceObject.TryLoad();
 		}
 		return nullptr;
 	}
+
 	/**
-	 * Override it to setup the UObject source data (payload data) and configure the properties asynchronously. Do not create the asset UObject asynchronously.
+	 * Override it to end import/re-import on the game thread. Helpful if you need to display ui (i.e. re-import material conflict resolution) or if you
+	 * need to do anything not thread safe to complete the import/re-import.
 	 *
 	 * @param Arguments - The structure containing all necessary data to imported the UObject, see the structure definition for the documentation.
 	 * @return the imported UObject or nullptr if there is an error.
 	 */
-	virtual UObject* ImportAssetObject_Async(const FImportAssetObjectParams& Arguments)
+	virtual UObject* EndImportAssetObject_GameThread(const FImportAssetObjectParams& Arguments)
 	{
-		PRAGMA_DISABLE_DEPRECATION_WARNINGS
-		return CreateAsset(Arguments);
-		PRAGMA_ENABLE_DEPRECATION_WARNINGS
+		//By default simply return the UObject created by BeginImportAssetObject_GameThread that was store in the asset node.
+		FSoftObjectPath ReferenceObject;
+		if (Arguments.AssetNode && Arguments.AssetNode->GetCustomReferenceObject(ReferenceObject))
+		{
+			return ReferenceObject.TryLoad();
+		}
+		return nullptr;
 	}
 
 	/**
@@ -153,12 +175,6 @@ public:
 		const UInterchangeBaseNodeContainer* NodeContainer = nullptr;
 	};
 
-	UE_DEPRECATED(5.2, "This function is replaced by ImportSceneObject_GameThread.")
-	virtual UObject* CreateSceneObject(const FImportSceneObjectsParams& Arguments)
-	{
-		return nullptr;
-	}
-
 	/**
 	 * Creates the scene object from a Scene Node data.
 	 *
@@ -167,9 +183,7 @@ public:
 	 */
 	virtual UObject* ImportSceneObject_GameThread(const FImportSceneObjectsParams& Arguments)
 	{
-		PRAGMA_DISABLE_DEPRECATION_WARNINGS
-		return CreateSceneObject(Arguments);
-		PRAGMA_ENABLE_DEPRECATION_WARNINGS
+		return nullptr;
 	}
 
 	/** Call when the user cancel the operation. */
@@ -196,38 +210,22 @@ public:
 ;
 	};
 
-	UE_DEPRECATED(5.2, "This function is replaced by SetupObject_GameThread.")
-	virtual void PreImportPreCompletedCallback(const FSetupObjectParams& Arguments)
-	{
-		check(IsInGameThread());
-	}
-
 	/*
 	 * Do any UObject setup required before the build (before PostEditChange) and after all dependency UObject have been imported.
 	 * @note - This function is called when starting the pre completion task (before PostEditChange is called for the asset).
 	 */
 	virtual void SetupObject_GameThread(const FSetupObjectParams& Arguments)
-	{
-		PRAGMA_DISABLE_DEPRECATION_WARNINGS
-		PreImportPreCompletedCallback(Arguments);
-		PRAGMA_ENABLE_DEPRECATION_WARNINGS
-	}
-	
-	UE_DEPRECATED(5.2, "This function is replace by FinalizeObject_GameThread.")
-	virtual void PostImportPreCompletedCallback(const FSetupObjectParams& Arguments)
 	{
 		check(IsInGameThread());
 	}
-
+	
 	/*
 	 * Do any final UObject setup after the build (after PostEditChange)
 	 * @note - This function is called at the end of the pre completion task (after PostEditChange is called for the asset).
 	 */
 	virtual void FinalizeObject_GameThread(const FSetupObjectParams& Arguments)
 	{
-		PRAGMA_DISABLE_DEPRECATION_WARNINGS
-		PostImportPreCompletedCallback(Arguments);
-		PRAGMA_ENABLE_DEPRECATION_WARNINGS
+		check(IsInGameThread());
 	}
 
 	/**
diff --git a/Engine/Source/Runtime/Interchange/Engine/Private/InterchangeManager.cpp b/Engine/Source/Runtime/Interchange/Engine/Private/InterchangeManager.cpp
--- a/Engine/Source/Runtime/Interchange/Engine/Private/InterchangeManager.cpp
+++ b/Engine/Source/Runtime/Interchange/Engine/Private/InterchangeManager.cpp
@@ -313,8 +313,9 @@ FGraphEventArray UE::Interchange::FImportAsyncHelper::GetCompletionTaskGraphEven
 	FTaskGraphInterface::Get().WaitUntilTasksComplete(TasksToComplete, ENamedThreads::GameThread);
 	TasksToComplete.Reset();
 
-	TasksToComplete.Append(CreatePackageTasks);
-	TasksToComplete.Append(CreateAssetTasks);
+	TasksToComplete.Append(BeginImportObjectTasks);
+	TasksToComplete.Append(ImportObjectTasks);
+	TasksToComplete.Append(FinalizeImportObjectTasks);
 	TasksToComplete.Append(SceneTasks);
 	TasksToComplete.Append(PipelinePostImportTasks);
 
@@ -641,11 +642,21 @@ UInterchangeManager& UInterchangeManager::GetInterchangeManager()
 		//We cancel any running task when we pre exit the engine
 		FCoreDelegates::OnEnginePreExit.AddLambda([]()
 		{
-			//If a user run a editor commandlet, we want to finish the import before the editor close.
-			//In editor mode, the user cannot close the editor if an import task is running, so we should not endup here.
-			const bool bCancel = !GIsEditor;
-			//Synchronously wait all task to finish
-			InterchangeManager->WaitUntilAllTasksDone(bCancel);
+			//InterchangeManager should be valid at this point since this lambda is where the strong reference pointer get reset.
+			if (!ensure(InterchangeManager.IsValid()))
+			{
+				InterchangeManagerScopeOfLifeEnded = true;
+				return;
+			}
+
+			if (IsInterchangeImportEnabled())
+			{
+				//If a user run a editor commandlet, we want to finish the import before the editor close.
+				//In editor mode, the user cannot close the editor if an import task is running, so we should not endup here.
+				const bool bCancel = !GIsEditor;
+				//Synchronously wait all task to finish
+				InterchangeManager->WaitUntilAllTasksDone(bCancel);
+			}
 
 			//Task should have been cancel in the Engine pre exit callback
 			ensure(InterchangeManager->ImportTasks.Num() == 0);
diff --git a/Engine/Source/Runtime/Interchange/Engine/Private/Tasks/InterchangeTaskImportObject.cpp b/Engine/Source/Runtime/Interchange/Engine/Private/Tasks/InterchangeTaskImportObject.cpp
--- a/Engine/Source/Runtime/Interchange/Engine/Private/Tasks/InterchangeTaskImportObject.cpp
+++ b/Engine/Source/Runtime/Interchange/Engine/Private/Tasks/InterchangeTaskImportObject.cpp
@@ -23,153 +23,243 @@
 #include "UObject/ObjectMacros.h"
 #include "UObject/WeakObjectPtrTemplates.h"
 
-namespace UE
+namespace UE::Interchange::Private
 {
-	namespace Interchange
+	void InternalGetPackageName(const UE::Interchange::FImportAsyncHelper& AsyncHelper, const int32 SourceIndex, const FString& PackageBasePath, const UInterchangeFactoryBaseNode* FactoryNode, FString& OutPackageName, FString& OutAssetName)
 	{
-		namespace Private
-		{
-			void InternalGetPackageName(const UE::Interchange::FImportAsyncHelper& AsyncHelper, const int32 SourceIndex, const FString& PackageBasePath, const UInterchangeFactoryBaseNode* FactoryNode, FString& OutPackageName, FString& OutAssetName)
-			{
-				TRACE_CPUPROFILER_EVENT_SCOPE("UE::Interchange::Private::InternalGetPackageName")
-				const UInterchangeSourceData* SourceData = AsyncHelper.SourceDatas[SourceIndex];
-				check(SourceData);
-				FString NodeDisplayName = FactoryNode->GetAssetName();
+		TRACE_CPUPROFILER_EVENT_SCOPE("UE::Interchange::Private::InternalGetPackageName")
+		const UInterchangeSourceData* SourceData = AsyncHelper.SourceDatas[SourceIndex];
+		check(SourceData);
+		FString NodeDisplayName = FactoryNode->GetAssetName();
 
-				// Set the asset name and the package name
-				OutAssetName = NodeDisplayName;
-				SanitizeObjectName(OutAssetName);
+		// Set the asset name and the package name
+		OutAssetName = NodeDisplayName;
+		SanitizeObjectName(OutAssetName);
 
-				FString SanitizedPackageBasePath = PackageBasePath;
-				SanitizeObjectPath(SanitizedPackageBasePath);
+		FString SanitizedPackageBasePath = PackageBasePath;
+		SanitizeObjectPath(SanitizedPackageBasePath);
 
-				FString SubPath;
-				if (FactoryNode->GetCustomSubPath(SubPath))
-				{
-					SanitizeObjectPath(SubPath);
-				}
+		FString SubPath;
+		if (FactoryNode->GetCustomSubPath(SubPath))
+		{
+			SanitizeObjectPath(SubPath);
+		}
 
-				OutPackageName = FPaths::Combine(*SanitizedPackageBasePath, *SubPath, *OutAssetName);
-			}
-			bool ShouldReimportFactoryNode(UInterchangeFactoryBaseNode* FactoryNode, const UInterchangeBaseNodeContainer* NodeContainer, UObject* ReimportObject)
-			{
-				TRACE_CPUPROFILER_EVENT_SCOPE("UE::Interchange::Private::ShouldReimportFactoryNode")
+		OutPackageName = FPaths::Combine(*SanitizedPackageBasePath, *SubPath, *OutAssetName);
+	}
+	bool ShouldReimportFactoryNode(UInterchangeFactoryBaseNode* FactoryNode, const UInterchangeBaseNodeContainer* NodeContainer, UObject* ReimportObject)
+	{
+		TRACE_CPUPROFILER_EVENT_SCOPE("UE::Interchange::Private::ShouldReimportFactoryNode")
 
-				if (!NodeContainer)
+		if (!NodeContainer)
+		{
+			return false;
+		}
+		//Find all potential factory node
+		TArray<UInterchangeFactoryBaseNode*> PotentialFactoryNodes;
+		UClass* FactoryClass = FactoryNode->GetObjectClass();
+		NodeContainer->IterateNodesOfType<UInterchangeFactoryBaseNode>([FactoryClass, &PotentialFactoryNodes](const FString& NodeUniqueID, UInterchangeFactoryBaseNode* CurrentFactoryNode)
+			{
+				if (UClass* CurrentFactoryClass = CurrentFactoryNode->GetObjectClass())
 				{
-					return false;
-				}
-				//Find all potential factory node
-				TArray<UInterchangeFactoryBaseNode*> PotentialFactoryNodes;
-				UClass* FactoryClass = FactoryNode->GetObjectClass();
-				NodeContainer->IterateNodesOfType<UInterchangeFactoryBaseNode>([FactoryClass, &PotentialFactoryNodes](const FString& NodeUniqueID, UInterchangeFactoryBaseNode* CurrentFactoryNode)
+					if (CurrentFactoryClass->IsChildOf(FactoryClass))
 					{
-						if (UClass* CurrentFactoryClass = CurrentFactoryNode->GetObjectClass())
-						{
-							if (CurrentFactoryClass->IsChildOf(FactoryClass))
-							{
-								PotentialFactoryNodes.Add(CurrentFactoryNode);
-							}
-						}
-					});
-
-				if (PotentialFactoryNodes.Num() == 1)
-				{
-					//There is only one factory node that will generate this class of UObject, no need to match the unique id or the name.
-					ensure(PotentialFactoryNodes[0] == FactoryNode);
-					return true;
+						PotentialFactoryNodes.Add(CurrentFactoryNode);
+					}
 				}
+			});
+		
+		if (PotentialFactoryNodes.Num() == 1)
+		{
+			//There is only one factory node that will generate this class of UObject, no need to match the unique id or the name.
+			ensure(PotentialFactoryNodes[0] == FactoryNode);
+			//The class of the object must fit with the factory object class.
+			return ReimportObject->GetClass()->IsChildOf(FactoryClass);
+		}
 
-				//If the source was used to import multiple asset, see if the FactoryNode match the ReimportObject.
-				TArray<UObject*> SubObjects;
-				GetObjectsWithOuter(ReimportObject, SubObjects);
-				for (UObject* SubObject : SubObjects)
+		//See if the FactoryNode match the ReimportObject.
+		TArray<UObject*> SubObjects;
+		GetObjectsWithOuter(ReimportObject, SubObjects);
+		for (UObject* SubObject : SubObjects)
+		{
+			if(UInterchangeAssetImportData* OriginalAssetImportData = Cast<UInterchangeAssetImportData>(SubObject))
+			{
+				if (UInterchangeBaseNodeContainer* OriginalNodeContainer = OriginalAssetImportData->NodeContainer)
 				{
-					if(UInterchangeAssetImportData* OriginalAssetImportData = Cast<UInterchangeAssetImportData>(SubObject))
+					//Find the original factory node used by the last ReimportObject import
+					if (UInterchangeFactoryBaseNode* OriginalFactoryNode = OriginalNodeContainer->GetFactoryNode(OriginalAssetImportData->NodeUniqueID))
 					{
-						if (UInterchangeBaseNodeContainer* OriginalNodeContainer = OriginalAssetImportData->NodeContainer)
+						//Compare the original factory node UObject class to the factory node UObject class
+						if (OriginalFactoryNode->GetObjectClass()->IsChildOf(FactoryClass))
 						{
-							//Find the original factory node used by the last ReimportObject import
-							if (UInterchangeFactoryBaseNode* OriginalFactoryNode = OriginalNodeContainer->GetFactoryNode(OriginalAssetImportData->NodeUniqueID))
+							//Compare the original factory node, unique id to the factory node unique id
+							if (OriginalFactoryNode->GetUniqueID().Equals(FactoryNode->GetUniqueID()))
 							{
-								//Compare the original factory node UObject class to the factory node UObject class
-								if (OriginalFactoryNode->GetObjectClass()->IsChildOf(FactoryClass))
+								return true;
+							}
+							//Compare the original factory node, name to the factory node name
+							if (OriginalFactoryNode->GetDisplayLabel().Equals(FactoryNode->GetDisplayLabel()))
+							{
+								FString PackageSubPath;
+								FactoryNode->GetCustomSubPath(PackageSubPath);
+								FString OriginalPackageSubPath;
+								OriginalFactoryNode->GetCustomSubPath(PackageSubPath);
+								//Make sure both sub path are equal
+								if (PackageSubPath.Equals(OriginalPackageSubPath))
 								{
-									//Compare the original factory node, unique id to the factory node unique id
-									if (OriginalFactoryNode->GetUniqueID().Equals(FactoryNode->GetUniqueID()))
-									{
-										return true;
-									}
-									//Compare the original factory node, name to the factory node name
-									if (OriginalFactoryNode->GetDisplayLabel().Equals(FactoryNode->GetDisplayLabel()))
-									{
-										FString PackageSubPath;
-										FactoryNode->GetCustomSubPath(PackageSubPath);
-										FString OriginalPackageSubPath;
-										OriginalFactoryNode->GetCustomSubPath(PackageSubPath);
-										//Make sure both sub path are equal
-										if (PackageSubPath.Equals(OriginalPackageSubPath))
-										{
-											return true;
-										}
-									}
+									return true;
 								}
 							}
 						}
 					}
 				}
-				return false;
 			}
+		}
+		return false;
+	}
 
-			bool CanImportClass(UE::Interchange::FImportAsyncHelper& AsyncHelper, UInterchangeFactoryBaseNode& FactoryNode, int32 SourceIndex)
+	bool CanImportClass(UE::Interchange::FImportAsyncHelper& AsyncHelper, UInterchangeFactoryBaseNode& FactoryNode, int32 SourceIndex)
+	{
+	#if WITH_EDITOR
+		if (UClass* Class = FactoryNode.GetObjectClass())
+		{
+			if (AsyncHelper.AllowedClasses.Contains(Class))
 			{
-			#if WITH_EDITOR
-				if (UClass* Class = FactoryNode.GetObjectClass())
+				return true;
+			}
+			else if (AsyncHelper.DeniedClasses.Contains(Class))
+			{
+				return false;
+			}
+			else
+			{
+				IAssetTools& AssetTools = FAssetToolsModule::GetModule().Get();
+				TSharedPtr<FPathPermissionList> AssetClassPermissionList = AssetTools.GetAssetClassPathPermissionList(EAssetClassAction::ImportAsset);
+				if (AssetClassPermissionList && AssetClassPermissionList->HasFiltering())
 				{
-					if (AsyncHelper.AllowedClasses.Contains(Class))
-					{
-						return true;
-					}
-					else if (AsyncHelper.DeniedClasses.Contains(Class))
+					if (!AssetClassPermissionList->PassesFilter(Class->GetPathName()))
 					{
+						UE_LOG(LogInterchangeEngine, Display, TEXT("The creation of asset of class '%s' is not allowed in this project."), *Class->GetName());
+						AsyncHelper.DeniedClasses.Add(Class);
 						return false;
 					}
-					else
-					{
-						IAssetTools& AssetTools = FAssetToolsModule::GetModule().Get();
-						TSharedPtr<FPathPermissionList> AssetClassPermissionList = AssetTools.GetAssetClassPathPermissionList(EAssetClassAction::ImportAsset);
-						if (AssetClassPermissionList && AssetClassPermissionList->HasFiltering())
-						{
-							if (!AssetClassPermissionList->PassesFilter(Class->GetPathName()))
-							{
-								UE_LOG(LogInterchangeEngine, Display, TEXT("The creation of asset of class '%s' is not allowed in this project."), *Class->GetName());
-								AsyncHelper.DeniedClasses.Add(Class);
-								return false;
-							}
-						}
-
-						AsyncHelper.AllowedClasses.Add(Class);
-
-						return true;
-					}
 				}
-				else
-				{
-					return false;
-				}
-			#endif //WITH_EDITOR
+
+				AsyncHelper.AllowedClasses.Add(Class);
 
 				return true;
 			}
-		}//ns Private
-	}//ns Interchange
-}//ns UE
+		}
+		else
+		{
+			return false;
+		}
+	#endif //WITH_EDITOR
+
+		return true;
+	}
+
+	UObject* InternalImportObjectStartup(TSharedPtr<FImportAsyncHelper, ESPMode::ThreadSafe> AsyncHelper
+		, UInterchangeFactoryBaseNode* FactoryNode
+		, int32 SourceIndex
+		, const FString& PackageBasePath
+		, TUniqueFunction<UObject*(UInterchangeFactoryBase::FImportAssetObjectParams&)>FactoryOperationLambda)
+	{
+		check(AsyncHelper.IsValid());
+
+		//Verify if the task was cancel
+		if (AsyncHelper->bCancel || !FactoryNode || !Private::CanImportClass(*AsyncHelper, *FactoryNode, SourceIndex))
+		{
+			return nullptr;
+		}
+
+		UInterchangeFactoryBase* Factory = nullptr;
+		{
+			FScopeLock Lock(&AsyncHelper->CreatedFactoriesLock);
+			Factory = AsyncHelper->CreatedFactories.FindChecked(FactoryNode->GetUniqueID());
+		}
+
+		UPackage* Pkg = nullptr;
+		FString PackageName;
+		FString AssetName;
+		Private::InternalGetPackageName(*AsyncHelper, SourceIndex, PackageBasePath, FactoryNode, PackageName, AssetName);
+		bool bSkipAsset = false;
+		UObject* ReimportObject = AsyncHelper->TaskData.ReimportObject;
+		if (ReimportObject)
+		{
+			UInterchangeBaseNodeContainer* NodeContainer = nullptr;
+			if (AsyncHelper->BaseNodeContainers.IsValidIndex(SourceIndex))
+			{
+				NodeContainer = AsyncHelper->BaseNodeContainers[SourceIndex].Get();
+			}
+
+			if (Private::ShouldReimportFactoryNode(FactoryNode, NodeContainer, ReimportObject))
+			{
+				Pkg = ReimportObject->GetPackage();
+				PackageName = Pkg->GetPathName();
+				AssetName = ReimportObject->GetName();
+			}
+			else
+			{
+				bSkipAsset = true;
+			}
+		}
+		else
+		{
+			FScopeLock Lock(&AsyncHelper->CreatedPackagesLock);
+			UPackage** PkgPtr = AsyncHelper->CreatedPackages.Find(PackageName);
+
+			if (!PkgPtr || !(*PkgPtr))
+			{
+				UInterchangeResultError_Generic* Message = Factory->AddMessage<UInterchangeResultError_Generic>();
+				Message->SourceAssetName = AsyncHelper->SourceDatas[SourceIndex]->GetFilename();
+				Message->DestinationAssetName = AssetName;
+				Message->AssetType = FactoryNode->GetObjectClass();
+				Message->Text = NSLOCTEXT("Interchange", "BadPackage", "It was not possible to create the asset as its package was not created correctly.");
+
+				return nullptr;
+			}
+
+			if (!AsyncHelper->SourceDatas.IsValidIndex(SourceIndex) || !AsyncHelper->Translators.IsValidIndex(SourceIndex))
+			{
+				UInterchangeResultError_Generic* Message = Factory->AddMessage<UInterchangeResultError_Generic>();
+				Message->DestinationAssetName = AssetName;
+				Message->AssetType = FactoryNode->GetObjectClass();
+				Message->Text = NSLOCTEXT("Interchange", "SourceDataOrTranslatorInvalid", "It was not possible to create the asset as its translator was not created correctly.");
+
+				return nullptr;
+			}
+
+			Pkg = *PkgPtr;
+		}
+
+		if (!bSkipAsset)
+		{
+			UInterchangeTranslatorBase* Translator = AsyncHelper->Translators[SourceIndex];
+			//Import Asset describe by the node
+			UInterchangeFactoryBase::FImportAssetObjectParams CreateAssetParams;
+			CreateAssetParams.AssetName = AssetName;
+			CreateAssetParams.AssetNode = FactoryNode;
+			CreateAssetParams.Parent = Pkg;
+			CreateAssetParams.SourceData = AsyncHelper->SourceDatas[SourceIndex];
+			CreateAssetParams.Translator = Translator;
+			if (AsyncHelper->BaseNodeContainers.IsValidIndex(SourceIndex))
+			{
+				CreateAssetParams.NodeContainer = AsyncHelper->BaseNodeContainers[SourceIndex].Get();
+			}
+			CreateAssetParams.ReimportObject = ReimportObject;
+
+			return FactoryOperationLambda(CreateAssetParams);
+		}
+		return nullptr;
+	}
+}//ns UE::Interchange::Private
 
 void UE::Interchange::FTaskImportObject_GameThread::DoTask(ENamedThreads::Type CurrentThread, const FGraphEventRef& MyCompletionGraphEvent)
 {
 	TRACE_CPUPROFILER_EVENT_SCOPE("UE::Interchange::FTaskImportObject_GameThread::DoTask")
 #if INTERCHANGE_TRACE_ASYNCHRONOUS_TASK_ENABLED
-	INTERCHANGE_TRACE_ASYNCHRONOUS_TASK(CreatePackage)
+	INTERCHANGE_TRACE_ASYNCHRONOUS_TASK(TaskImportObject_GameThread)
 #endif
 	using namespace UE::Interchange;
 
@@ -225,7 +315,7 @@ void UE::Interchange::FTaskImportObject_GameThread::DoTask(ENamedThreads::Type C
 			CreateAssetParams.ReimportObject = ReimportObject;
 			FactoryNode->SetCustomReferenceObject(FSoftObjectPath(ReimportObject));
 			//We call CreateEmptyAsset to ensure any resource use by an existing UObject is released on the game thread
-			Factory->ImportAssetObject_GameThread(CreateAssetParams);
+			Factory->BeginImportAssetObject_GameThread(CreateAssetParams);
 		}
 		else
 		{
@@ -277,7 +367,7 @@ void UE::Interchange::FTaskImportObject_GameThread::DoTask(ENamedThreads::Type C
 		}
 		CreateAssetParams.ReimportObject = ReimportObject;
 		//Make sure the asset UObject is created with the correct type on the main thread
-		UObject* NodeAsset = Factory->ImportAssetObject_GameThread(CreateAssetParams);
+		UObject* NodeAsset = Factory->BeginImportAssetObject_GameThread(CreateAssetParams);
 		if (NodeAsset)
 		{
 			if (!NodeAsset->HasAnyInternalFlags(EInternalObjectFlags::Async))
@@ -309,9 +399,9 @@ void UE::Interchange::FTaskImportObject_GameThread::DoTask(ENamedThreads::Type C
 void UE::Interchange::FTaskImportObject_Async::DoTask(ENamedThreads::Type CurrentThread, const FGraphEventRef& MyCompletionGraphEvent)
 {
 #if INTERCHANGE_TRACE_ASYNCHRONOUS_TASK_ENABLED
-	INTERCHANGE_TRACE_ASYNCHRONOUS_TASK(CreateAsset)
+	INTERCHANGE_TRACE_ASYNCHRONOUS_TASK(TaskImportObject_Async)
 #endif
-		using namespace UE::Interchange;
+	using namespace UE::Interchange;
 
 	TSharedPtr<FImportAsyncHelper, ESPMode::ThreadSafe> AsyncHelper = WeakAsyncHelper.Pin();
 	check(AsyncHelper.IsValid());
@@ -328,116 +418,77 @@ void UE::Interchange::FTaskImportObject_Async::DoTask(ENamedThreads::Type Curren
 		Factory = AsyncHelper->CreatedFactories.FindChecked(FactoryNode->GetUniqueID());
 	}
 
-	UPackage* Pkg = nullptr;
-	FString PackageName;
-	FString AssetName;
-	Private::InternalGetPackageName(*AsyncHelper, SourceIndex, PackageBasePath, FactoryNode, PackageName, AssetName);
-	bool bSkipAsset = false;
-	UObject* ReimportObject = AsyncHelper->TaskData.ReimportObject;
-	if (ReimportObject)
-	{
-		UInterchangeBaseNodeContainer* NodeContainer = nullptr;
-		if (AsyncHelper->BaseNodeContainers.IsValidIndex(SourceIndex))
+	Private::InternalImportObjectStartup(AsyncHelper
+		, FactoryNode
+		, SourceIndex
+		, PackageBasePath
+		, [&Factory](UInterchangeFactoryBase::FImportAssetObjectParams& ImportAssetObjectParams)
 		{
-			NodeContainer = AsyncHelper->BaseNodeContainers[SourceIndex].Get();
-		}
-
-		if (Private::ShouldReimportFactoryNode(FactoryNode, NodeContainer, ReimportObject))
-		{
-			Pkg = ReimportObject->GetPackage();
-			PackageName = Pkg->GetPathName();
-			AssetName = ReimportObject->GetName();
-		}
-		else
-		{
-			bSkipAsset = true;
-		}
-	}
-	else
-	{
-		FScopeLock Lock(&AsyncHelper->CreatedPackagesLock);
-		UPackage** PkgPtr = AsyncHelper->CreatedPackages.Find(PackageName);
-
-		if (!PkgPtr || !(*PkgPtr))
-		{
-			UInterchangeResultError_Generic* Message = Factory->AddMessage<UInterchangeResultError_Generic>();
-			Message->SourceAssetName = AsyncHelper->SourceDatas[SourceIndex]->GetFilename();
-			Message->DestinationAssetName = AssetName;
-			Message->AssetType = FactoryNode->GetObjectClass();
-			Message->Text = NSLOCTEXT("Interchange", "BadPackage", "It was not possible to create the asset as its package was not created correctly.");
-
-			return;
-		}
+			return Factory->ImportAssetObject_Async(ImportAssetObjectParams);
+		});
+}
 
-		if (!AsyncHelper->SourceDatas.IsValidIndex(SourceIndex) || !AsyncHelper->Translators.IsValidIndex(SourceIndex))
-		{
-			UInterchangeResultError_Generic* Message = Factory->AddMessage<UInterchangeResultError_Generic>();
-			Message->DestinationAssetName = AssetName;
-			Message->AssetType = FactoryNode->GetObjectClass();
-			Message->Text = NSLOCTEXT("Interchange", "SourceDataOrTranslatorInvalid", "It was not possible to create the asset as its translator was not created correctly.");
+void UE::Interchange::FTaskImportObjectFinalize_GameThread::DoTask(ENamedThreads::Type CurrentThread, const FGraphEventRef& MyCompletionGraphEvent)
+{
+#if INTERCHANGE_TRACE_ASYNCHRONOUS_TASK_ENABLED
+	INTERCHANGE_TRACE_ASYNCHRONOUS_TASK(TaskImportObjectFinalize_GameThread)
+#endif
+	using namespace UE::Interchange;
 
-			return;
-		}
+	TSharedPtr<FImportAsyncHelper, ESPMode::ThreadSafe> AsyncHelper = WeakAsyncHelper.Pin();
+	check(AsyncHelper.IsValid());
 
-		Pkg = *PkgPtr;
+	//Verify if the task was cancel
+	if (AsyncHelper->bCancel || !FactoryNode || !Private::CanImportClass(*AsyncHelper, *FactoryNode, SourceIndex))
+	{
+		return;
 	}
 
-	UObject* NodeAsset = nullptr;
-	if (bSkipAsset)
+	UInterchangeFactoryBase* Factory = nullptr;
 	{
-		NodeAsset = nullptr;
+		FScopeLock Lock(&AsyncHelper->CreatedFactoriesLock);
+		Factory = AsyncHelper->CreatedFactories.FindChecked(FactoryNode->GetUniqueID());
 	}
-	else
-	{
-		UInterchangeTranslatorBase* Translator = AsyncHelper->Translators[SourceIndex];
-		//Import Asset describe by the node
-		UInterchangeFactoryBase::FImportAssetObjectParams CreateAssetParams;
-		CreateAssetParams.AssetName = AssetName;
-		CreateAssetParams.AssetNode = FactoryNode;
-		CreateAssetParams.Parent = Pkg;
-		CreateAssetParams.SourceData = AsyncHelper->SourceDatas[SourceIndex];
-		CreateAssetParams.Translator = Translator;
-		if (AsyncHelper->BaseNodeContainers.IsValidIndex(SourceIndex))
+
+	UObject* NodeAsset = Private::InternalImportObjectStartup(AsyncHelper
+		, FactoryNode
+		, SourceIndex
+		, PackageBasePath
+		, [&Factory](UInterchangeFactoryBase::FImportAssetObjectParams& ImportAssetObjectParams)
 		{
-			CreateAssetParams.NodeContainer = AsyncHelper->BaseNodeContainers[SourceIndex].Get();
-		}
-		CreateAssetParams.ReimportObject = ReimportObject;
+			return Factory->EndImportAssetObject_GameThread(ImportAssetObjectParams);
+		});
 
-		NodeAsset = Factory->ImportAssetObject_Async(CreateAssetParams);
-	}
 	if (NodeAsset)
 	{
-		if (!bSkipAsset)
-		{
-			FScopeLock Lock(&AsyncHelper->ImportedAssetsPerSourceIndexLock);
-			TArray<FImportAsyncHelper::FImportedObjectInfo>& ImportedInfos = AsyncHelper->ImportedAssetsPerSourceIndex.FindOrAdd(SourceIndex);
-			FImportAsyncHelper::FImportedObjectInfo* AssetInfoPtr = ImportedInfos.FindByPredicate([NodeAsset](const FImportAsyncHelper::FImportedObjectInfo& CurInfo)
+		FScopeLock Lock(&AsyncHelper->ImportedAssetsPerSourceIndexLock);
+		TArray<FImportAsyncHelper::FImportedObjectInfo>& ImportedInfos = AsyncHelper->ImportedAssetsPerSourceIndex.FindOrAdd(SourceIndex);
+		FImportAsyncHelper::FImportedObjectInfo* AssetInfoPtr = ImportedInfos.FindByPredicate([NodeAsset](const FImportAsyncHelper::FImportedObjectInfo& CurInfo)
 			{
 				return CurInfo.ImportedObject == NodeAsset;
 			});
 
-			if (!AssetInfoPtr)
-			{
-				FImportAsyncHelper::FImportedObjectInfo& AssetInfo = ImportedInfos.AddDefaulted_GetRef();
-				AssetInfo.ImportedObject = NodeAsset;
-				AssetInfo.Factory = Factory;
-				AssetInfo.FactoryNode = FactoryNode;
-				AssetInfo.bIsReimport = bool(ReimportObject != nullptr);
-			}
+		if (!AssetInfoPtr)
+		{
+			FImportAsyncHelper::FImportedObjectInfo& AssetInfo = ImportedInfos.AddDefaulted_GetRef();
+			AssetInfo.ImportedObject = NodeAsset;
+			AssetInfo.Factory = Factory;
+			AssetInfo.FactoryNode = FactoryNode;
+			AssetInfo.bIsReimport = bool(AsyncHelper->TaskData.ReimportObject != nullptr);
+		}
 
-			// Fill in destination asset and type in any results which have been added previously by a translator or pipeline, now that we have a corresponding factory.
-			UInterchangeResultsContainer* Results = AsyncHelper->AssetImportResult->GetResults();
-			for (UInterchangeResult* Result : Results->GetResults())
+		// Fill in destination asset and type in any results which have been added previously by a translator or pipeline, now that we have a corresponding factory.
+		UInterchangeResultsContainer* Results = AsyncHelper->AssetImportResult->GetResults();
+		for (UInterchangeResult* Result : Results->GetResults())
+		{
+			if (!Result->InterchangeKey.IsEmpty() && (Result->DestinationAssetName.IsEmpty() || Result->AssetType == nullptr))
 			{
-				if (!Result->InterchangeKey.IsEmpty() && (Result->DestinationAssetName.IsEmpty() || Result->AssetType == nullptr))
+				TArray<FString> TargetAssets;
+				FactoryNode->GetTargetNodeUids(TargetAssets);
+				if (TargetAssets.Contains(Result->InterchangeKey))
 				{
-					TArray<FString> TargetAssets;
-					FactoryNode->GetTargetNodeUids(TargetAssets);
-					if (TargetAssets.Contains(Result->InterchangeKey))
-					{
-						Result->DestinationAssetName = NodeAsset->GetPathName();
-						Result->AssetType = NodeAsset->GetClass();
-					}
+					Result->DestinationAssetName = NodeAsset->GetPathName();
+					Result->AssetType = NodeAsset->GetClass();
 				}
 			}
 		}
diff --git a/Engine/Source/Runtime/Interchange/Engine/Private/Tasks/InterchangeTaskImportObject.h b/Engine/Source/Runtime/Interchange/Engine/Private/Tasks/InterchangeTaskImportObject.h
--- a/Engine/Source/Runtime/Interchange/Engine/Private/Tasks/InterchangeTaskImportObject.h
+++ b/Engine/Source/Runtime/Interchange/Engine/Private/Tasks/InterchangeTaskImportObject.h
@@ -10,89 +10,124 @@
 #include "UObject/WeakObjectPtrTemplates.h"
 #include "Nodes/InterchangeFactoryBaseNode.h"
 
-namespace UE
+namespace UE::Interchange
 {
-	namespace Interchange
+	/**
+		* This task create UPackage and UObject, Cook::PackageTracker::NotifyUObjectCreated is not thread safe, so we need to create the packages on the main thread
+		*/
+	class FTaskImportObject_GameThread
 	{
+	private:
+		FString PackageBasePath;
+		int32 SourceIndex;
+		TWeakPtr<FImportAsyncHelper, ESPMode::ThreadSafe> WeakAsyncHelper;
+		UInterchangeFactoryBaseNode* FactoryNode;
+		const UClass* FactoryClass;
+
+	public:
+		FTaskImportObject_GameThread(const FString& InPackageBasePath, const int32 InSourceIndex, TWeakPtr<FImportAsyncHelper, ESPMode::ThreadSafe> InAsyncHelper, UInterchangeFactoryBaseNode* InFactoryNode, const UClass* InFactoryClass)
+			: PackageBasePath(InPackageBasePath)
+			, SourceIndex(InSourceIndex)
+			, WeakAsyncHelper(InAsyncHelper)
+			, FactoryNode(InFactoryNode)
+			, FactoryClass(InFactoryClass)
+		{
+			check(FactoryNode);
+			check(FactoryClass);
+		}
+
+		ENamedThreads::Type GetDesiredThread() const
+		{
+			return ENamedThreads::GameThread;
+		}
+
+		static ESubsequentsMode::Type GetSubsequentsMode()
+		{
+			return ESubsequentsMode::TrackSubsequents;
+		}
+
+		TStatId GetStatId() const
+		{
+			RETURN_QUICK_DECLARE_CYCLE_STAT(FTaskImportObject_GameThread, STATGROUP_TaskGraphTasks);
+		}
+
+		void DoTask(ENamedThreads::Type CurrentThread, const FGraphEventRef& MyCompletionGraphEvent);
+	};
+
+	class FTaskImportObject_Async
+	{
+	private:
+		FString PackageBasePath;
+		int32 SourceIndex;
+		TWeakPtr<FImportAsyncHelper, ESPMode::ThreadSafe> WeakAsyncHelper;
+		UInterchangeFactoryBaseNode* FactoryNode;
+
+	public:
+		FTaskImportObject_Async(const FString& InPackageBasePath, const int32 InSourceIndex, TWeakPtr<FImportAsyncHelper, ESPMode::ThreadSafe> InAsyncHelper, UInterchangeFactoryBaseNode* InFactoryNode)
+			: PackageBasePath(InPackageBasePath)
+			, SourceIndex(InSourceIndex)
+			, WeakAsyncHelper(InAsyncHelper)
+			, FactoryNode(InFactoryNode)
+		{
+			check(FactoryNode);
+		}
+
+		ENamedThreads::Type GetDesiredThread() const
+		{
+			return ENamedThreads::AnyBackgroundThreadNormalTask;
+		}
 
-		/**
-		 * This task create UPackage and UObject, Cook::PackageTracker::NotifyUObjectCreated is not thread safe, so we need to create the packages on the main thread
-		 */
-		class FTaskImportObject_GameThread
+		static ESubsequentsMode::Type GetSubsequentsMode()
 		{
-		private:
-			FString PackageBasePath;
-			int32 SourceIndex;
-			TWeakPtr<FImportAsyncHelper, ESPMode::ThreadSafe> WeakAsyncHelper;
-			UInterchangeFactoryBaseNode* FactoryNode;
-			const UClass* FactoryClass;
-
-		public:
-			FTaskImportObject_GameThread(const FString& InPackageBasePath, const int32 InSourceIndex, TWeakPtr<FImportAsyncHelper, ESPMode::ThreadSafe> InAsyncHelper, UInterchangeFactoryBaseNode* InFactoryNode, const UClass* InFactoryClass)
-				: PackageBasePath(InPackageBasePath)
-				, SourceIndex(InSourceIndex)
-				, WeakAsyncHelper(InAsyncHelper)
-				, FactoryNode(InFactoryNode)
-				, FactoryClass(InFactoryClass)
-			{
-				check(FactoryNode);
-				check(FactoryClass);
-			}
-
-			ENamedThreads::Type GetDesiredThread() const
-			{
-				return ENamedThreads::GameThread;
-			}
-
-			static ESubsequentsMode::Type GetSubsequentsMode()
-			{
-				return ESubsequentsMode::TrackSubsequents;
-			}
-
-			TStatId GetStatId() const
-			{
-				RETURN_QUICK_DECLARE_CYCLE_STAT(FTaskImportObject_GameThread, STATGROUP_TaskGraphTasks);
-			}
-
-			void DoTask(ENamedThreads::Type CurrentThread, const FGraphEventRef& MyCompletionGraphEvent);
-		};
-
-		class FTaskImportObject_Async
+			return ESubsequentsMode::TrackSubsequents;
+		}
+
+		TStatId GetStatId() const
+		{
+			RETURN_QUICK_DECLARE_CYCLE_STAT(FTaskImportObject_Async, STATGROUP_TaskGraphTasks);
+		}
+
+		void DoTask(ENamedThreads::Type CurrentThread, const FGraphEventRef& MyCompletionGraphEvent);
+	};
+
+	/**
+		* This task create UPackage and UObject, Cook::PackageTracker::NotifyUObjectCreated is not thread safe, so we need to create the packages on the main thread
+		*/
+	class FTaskImportObjectFinalize_GameThread
+	{
+	private:
+		FString PackageBasePath;
+		int32 SourceIndex;
+		TWeakPtr<FImportAsyncHelper, ESPMode::ThreadSafe> WeakAsyncHelper;
+		UInterchangeFactoryBaseNode* FactoryNode;
+		const UClass* FactoryClass;
+
+	public:
+		FTaskImportObjectFinalize_GameThread(const FString& InPackageBasePath, const int32 InSourceIndex, TWeakPtr<FImportAsyncHelper, ESPMode::ThreadSafe> InAsyncHelper, UInterchangeFactoryBaseNode* InFactoryNode)
+			: PackageBasePath(InPackageBasePath)
+			, SourceIndex(InSourceIndex)
+			, WeakAsyncHelper(InAsyncHelper)
+			, FactoryNode(InFactoryNode)
+		{
+			check(FactoryNode);
+		}
+
+		ENamedThreads::Type GetDesiredThread() const
 		{
-		private:
-			FString PackageBasePath;
-			int32 SourceIndex;
-			TWeakPtr<FImportAsyncHelper, ESPMode::ThreadSafe> WeakAsyncHelper;
-			UInterchangeFactoryBaseNode* FactoryNode;
-
-		public:
-			FTaskImportObject_Async(const FString& InPackageBasePath, const int32 InSourceIndex, TWeakPtr<FImportAsyncHelper, ESPMode::ThreadSafe> InAsyncHelper, UInterchangeFactoryBaseNode* InFactoryNode)
-				: PackageBasePath(InPackageBasePath)
-				, SourceIndex(InSourceIndex)
-				, WeakAsyncHelper(InAsyncHelper)
-				, FactoryNode(InFactoryNode)
-			{
-				check(FactoryNode);
-			}
-
-			ENamedThreads::Type GetDesiredThread() const
-			{
-				return ENamedThreads::AnyBackgroundThreadNormalTask;
-			}
-
-			static ESubsequentsMode::Type GetSubsequentsMode()
-			{
-				return ESubsequentsMode::TrackSubsequents;
-			}
-
-			TStatId GetStatId() const
-			{
-				RETURN_QUICK_DECLARE_CYCLE_STAT(FTaskImportObject_Async, STATGROUP_TaskGraphTasks);
-			}
-
-			void DoTask(ENamedThreads::Type CurrentThread, const FGraphEventRef& MyCompletionGraphEvent);
-		};
-
-
-	} //ns Interchange
-}//ns UE
+			return ENamedThreads::GameThread;
+		}
+
+		static ESubsequentsMode::Type GetSubsequentsMode()
+		{
+			return ESubsequentsMode::TrackSubsequents;
+		}
+
+		TStatId GetStatId() const
+		{
+			RETURN_QUICK_DECLARE_CYCLE_STAT(FTaskImportObjectFinalize_GameThread, STATGROUP_TaskGraphTasks);
+		}
+
+		void DoTask(ENamedThreads::Type CurrentThread, const FGraphEventRef& MyCompletionGraphEvent);
+	};
+
+}//ns UE::Interchange
diff --git a/Engine/Source/Runtime/Interchange/Engine/Private/Tasks/InterchangeTaskParsing.cpp b/Engine/Source/Runtime/Interchange/Engine/Private/Tasks/InterchangeTaskParsing.cpp
--- a/Engine/Source/Runtime/Interchange/Engine/Private/Tasks/InterchangeTaskParsing.cpp
+++ b/Engine/Source/Runtime/Interchange/Engine/Private/Tasks/InterchangeTaskParsing.cpp
@@ -292,20 +292,26 @@ void UE::Interchange::FTaskParsing::DoTask(ENamedThreads::Type CurrentThread, co
 			if (ensureMsgf(!CreatedTasksAssetNames.Contains(AssetFullPath),
 				TEXT("Found multiple task data with the same asset name (%s). Only one will be executed."), *AssetFullPath))
 			{
-				//Add create package task has a prerequisite of FTaskImportObject_Async. Create package task is a game thread task
-				FGraphEventArray CreatePackagePrerequistes;
-				int32 CreatePackageTaskIndex = AsyncHelper->CreatePackageTasks.Add(
+				FGraphEventArray ImportObjectTasksPrerequistes;
+				int32 BeginImportObjectTaskIndex = AsyncHelper->BeginImportObjectTasks.Add(
 					TGraphTask<FTaskImportObject_GameThread>::CreateTask(&(TaskData.Prerequisites)).ConstructAndDispatchWhenReady(PackageBasePath, SourceIndex, WeakAsyncHelper, FactoryNode, FactoryClass)
 				);
-				CreatePackagePrerequistes.Add(AsyncHelper->CreatePackageTasks[CreatePackageTaskIndex]);
+				ImportObjectTasksPrerequistes.Add(AsyncHelper->BeginImportObjectTasks[BeginImportObjectTaskIndex]);
 
-				int32 CreateTaskIndex = AsyncHelper->CreateAssetTasks.Add(
-					TGraphTask<FTaskImportObject_Async>::CreateTask(&(CreatePackagePrerequistes)).ConstructAndDispatchWhenReady(PackageBasePath, SourceIndex, WeakAsyncHelper, FactoryNode)
+				int32 ImportObjectTaskIndex = AsyncHelper->ImportObjectTasks.Add(
+					TGraphTask<FTaskImportObject_Async>::CreateTask(&(ImportObjectTasksPrerequistes)).ConstructAndDispatchWhenReady(PackageBasePath, SourceIndex, WeakAsyncHelper, FactoryNode)
+				);
+
+				FGraphEventArray FinalizeImportObjectTasksPrerequistes;
+				FinalizeImportObjectTasksPrerequistes.Add(AsyncHelper->ImportObjectTasks[ImportObjectTaskIndex]);
+
+				int32 FinalizeCreateTaskIndex = AsyncHelper->FinalizeImportObjectTasks.Add(
+					TGraphTask<FTaskImportObjectFinalize_GameThread>::CreateTask(&(FinalizeImportObjectTasksPrerequistes)).ConstructAndDispatchWhenReady(PackageBasePath, SourceIndex, WeakAsyncHelper, FactoryNode)
 				);
 
 				CreatedTasksAssetNames.Add(AssetFullPath);
 
-				return AsyncHelper->CreateAssetTasks[CreateTaskIndex];
+				return AsyncHelper->FinalizeImportObjectTasks[FinalizeCreateTaskIndex];
 			}
 			else
 			{
diff --git a/Engine/Source/Runtime/Interchange/Engine/Public/InterchangeManager.h b/Engine/Source/Runtime/Interchange/Engine/Public/InterchangeManager.h
--- a/Engine/Source/Runtime/Interchange/Engine/Public/InterchangeManager.h
+++ b/Engine/Source/Runtime/Interchange/Engine/Public/InterchangeManager.h
@@ -202,8 +202,9 @@ namespace UE
 			TArray<FGraphEventRef> PipelineTasks;
 			TArray<FGraphEventRef> PipelinePostImportTasks;
 			FGraphEventRef ParsingTask;
-			TArray<FGraphEventRef> CreatePackageTasks;
-			TArray<FGraphEventRef> CreateAssetTasks;
+			TArray<FGraphEventRef> BeginImportObjectTasks;
+			TArray<FGraphEventRef> ImportObjectTasks;
+			TArray<FGraphEventRef> FinalizeImportObjectTasks;
 			TArray<FGraphEventRef> SceneTasks;
 
 			FGraphEventRef PreAsyncCompletionTask;
