- added USkeletalMesh::Get/CommitMeshDescription functions to ease skeletal mesh manipulations thru mesh description
- store RefBasesInvMatrix when transacting as the number of bones serialized with the ref skeleton can vary
- fixed transaction issue with Modify not being called on USkeletalMeshEditorData

diff --git a/Engine/Plugins/Runtime/MeshModelingToolset/Source/ModelingComponentsEditorOnly/Private/ToolTargets/SkeletalMeshToolTarget.cpp b/Engine/Plugins/Runtime/MeshModelingToolset/Source/ModelingComponentsEditorOnly/Private/ToolTargets/SkeletalMeshToolTarget.cpp
--- a/Engine/Plugins/Runtime/MeshModelingToolset/Source/ModelingComponentsEditorOnly/Private/ToolTargets/SkeletalMeshToolTarget.cpp
+++ b/Engine/Plugins/Runtime/MeshModelingToolset/Source/ModelingComponentsEditorOnly/Private/ToolTargets/SkeletalMeshToolTarget.cpp
@@ -149,28 +149,10 @@ FMeshDescription USkeletalMeshReadOnlyToolTarget::GetEmptyMeshDescription()
 	return EmptyMeshDescription;
 }
 
-void USkeletalMeshReadOnlyToolTarget::GetMeshDescription(const USkeletalMesh* SkeletalMeshIn, FMeshDescription& MeshDescription)
+void USkeletalMeshReadOnlyToolTarget::GetMeshDescription(const USkeletalMesh* SkeletalMeshIn, FMeshDescription& MeshDescriptionOut)
 {
 	using namespace USkeletalMeshToolTargetLocals;
-
-	// Check first if we have bulk data available and non-empty.
-	if (SkeletalMeshIn->IsLODImportedDataBuildAvailable(LODIndex) && !SkeletalMeshIn->IsLODImportedDataEmpty(LODIndex))
-	{
-		FSkeletalMeshImportData SkeletalMeshImportData;
-		SkeletalMeshIn->LoadLODImportedData(LODIndex, SkeletalMeshImportData);
-		SkeletalMeshImportData.GetMeshDescription(MeshDescription);
-	}
-	else
-	{
-		// Fall back on the LOD model directly if no bulk data exists. When we commit
-		// the mesh description, we override using the bulk data. This can happen for older
-		// skeletal meshes, from UE 4.24 and earlier.
-		const FSkeletalMeshModel* SkeletalMeshModel = SkeletalMeshIn->GetImportedModel();
-		if (SkeletalMeshModel && SkeletalMeshModel->LODModels.IsValidIndex(LODIndex))
-		{
-			SkeletalMeshModel->LODModels[LODIndex].GetMeshDescription(MeshDescription, SkeletalMeshIn);
-		}			
-	}
+	SkeletalMeshIn->GetMeshDescription(LODIndex, MeshDescriptionOut);
 }
 
 FDynamicMesh3 USkeletalMeshReadOnlyToolTarget::GetDynamicMesh()
@@ -218,25 +200,11 @@ void USkeletalMeshToolTarget::CommitMeshDescription(USkeletalMesh* SkeletalMeshI
 	// flush any pending rendering commands, which might touch a component while we are rebuilding it's mesh
 	FlushRenderingCommands();
 
-	// make sure transactional flag is on for this asset
-	SkeletalMeshIn->SetFlags(RF_Transactional);
-
-	verify(SkeletalMeshIn->Modify());
-
 	FCommitterParams CommitterParams;
-
 	CommitterParams.MeshDescriptionOut = MeshDescription;
-
 	Committer(CommitterParams);
-
-	FSkeletalMeshImportData SkeletalMeshImportData = 
-		FSkeletalMeshImportData::CreateFromMeshDescription(*CommitterParams.MeshDescriptionOut);
-	SkeletalMeshIn->SaveLODImportedData(LODIndex, SkeletalMeshImportData);
-
-	// Make sure the mesh builder knows it's the latest variety, so that the render data gets
-	// properly rebuilt.
-	SkeletalMeshIn->SetLODImportedDataVersions(LODIndex, ESkeletalMeshGeoImportVersions::LatestVersion, ESkeletalMeshSkinningImportVersions::LatestVersion);
-	SkeletalMeshIn->SetUseLegacyMeshDerivedDataKey(false);
+	
+	SkeletalMeshIn->CommitMeshDescription(LODIndex, *CommitterParams.MeshDescriptionOut);
 
 	SkeletalMeshIn->PostEditChange();
 }
diff --git a/Engine/Source/Developer/SkeletalMeshUtilitiesCommon/Private/SkeletonModifier.cpp b/Engine/Source/Developer/SkeletalMeshUtilitiesCommon/Private/SkeletonModifier.cpp
--- a/Engine/Source/Developer/SkeletalMeshUtilitiesCommon/Private/SkeletonModifier.cpp
+++ b/Engine/Source/Developer/SkeletalMeshUtilitiesCommon/Private/SkeletonModifier.cpp
@@ -16,6 +16,12 @@
 #include "AssetRegistry/AssetData.h"
 #include "Rendering/SkeletalMeshModel.h"
 
+
+namespace FSkeletonModifierLocals
+{
+	static constexpr int32 LODIndex = 0;
+}
+
 FTransform FMirrorOptions::MirrorTransform(const FTransform& InTransform) const
 {
 	FTransform Transform = InTransform;
@@ -144,8 +150,17 @@ bool FSkeletonModifier::Init(USkeletalMesh* InSkeletalMesh)
 		return false;
 	}
 
+	// verify user is not trying to modify one of the core engine assets
+	if (InSkeletalMesh->GetPathName().StartsWith(TEXT("/Engine/")))
+	{
+		UE_LOG(LogAnimation, Error, TEXT("Skeleton Modifier: Cannot modify built-in engine asset."));
+		return false;
+	}
+	
 	// check assets using this skeleton ?
 	{
+		// TODO avoid certain changes when the skeleton is referenced by other assets (i.e. changing the skeletal mesh's
+		// reference skeleton poses is fine, re-parenting/removing bones, etc. is not)
 		static const TArray<FTopLevelAssetPath> AssetPaths({
 		   UAnimSequence::StaticClass()->GetClassPathName(),
 		   UAnimMontage::StaticClass()->GetClassPathName(),
@@ -172,34 +187,17 @@ bool FSkeletonModifier::Init(USkeletalMesh* InSkeletalMesh)
 		}
 	}
 
-	// verify mesh has an imported model
-	const FSkeletalMeshModel* SkeletalMeshModel = InSkeletalMesh->GetImportedModel();
-	if (!SkeletalMeshModel)
-	{
-		UE_LOG(LogAnimation, Error, TEXT("Skeleton Modifier: Supplied skeletal mesh does not have an imported model."));
-		return false;
-	}
-	
-	// verify that we have a LOD 0 mesh
-	if (SkeletalMeshModel->LODModels.IsEmpty())
-	{
-		UE_LOG(LogAnimation, Error, TEXT("Skeleton Modifier: Supplied skeletal mesh does not have a LOD 0 mesh."));
-		return false;
-	}
-	
-	// verify user is not trying to modify one of the core engine assets
-	if (InSkeletalMesh->GetPathName().StartsWith(TEXT("/Engine/")))
-	{
-		UE_LOG(LogAnimation, Error, TEXT("Skeleton Modifier: Cannot modify built-in engine asset."));
-		return false;
-	}
-	
 	// store pointer to mesh and instantiate a mesh description for commiting changes
 	SkeletalMesh = InSkeletalMesh;
-	
+
+	// store mesh description to edit
 	MeshDescription = MakeUnique<FMeshDescription>();
-	const FSkeletalMeshLODModel& LOD0Mesh = SkeletalMeshModel->LODModels[0];
-	LOD0Mesh.GetMeshDescription(*MeshDescription, SkeletalMesh);
+	SkeletalMesh->GetMeshDescription(FSkeletonModifierLocals::LODIndex, *MeshDescription);
+	if (MeshDescription->IsEmpty())
+	{
+		UE_LOG(LogAnimation, Error, TEXT("Skeleton Modifier: mesh description is emtpy."));
+		return false;
+	}
 
 	// store reference skeleton to edit
 	ReferenceSkeleton = MakeUnique<FReferenceSkeleton>();
@@ -268,18 +266,6 @@ bool FSkeletonModifier::CommitSkeletonToSkeletalMesh()
 		}
 	}
 	
-	// update skeletal mesh
-	SkeletalMesh->SetRefSkeleton(*ReferenceSkeleton);
-	SkeletalMesh->GetRefBasesInvMatrix().Reset();
-	SkeletalMesh->CalculateInvRefMatrices(); 
-	
-	// update skeleton
-	USkeleton* Skeleton = SkeletalMesh->GetSkeleton();
-	if (Skeleton->RecreateBoneTree(SkeletalMesh))
-	{
-		Skeleton->MarkPackageDirty();	
-	}
-
 	// update mesh description
 	{
 		FSkeletalMeshAttributes MeshAttributes(*MeshDescription);
@@ -349,17 +335,28 @@ bool FSkeletonModifier::CommitSkeletonToSkeletalMesh()
 		}
 	}
 
-	// update skeletal mesh LOD (cf. USkeletalMeshToolTarget::CommitMeshDescription)
+	// update skeletal mesh
+	FlushRenderingCommands();
+	
+	SkeletalMesh->SetFlags(RF_Transactional);
+	SkeletalMesh->Modify();
+
+	// update the ref skeleton
+	SkeletalMesh->SetRefSkeleton(*ReferenceSkeleton);
+	SkeletalMesh->GetRefBasesInvMatrix().Reset();
+	SkeletalMesh->CalculateInvRefMatrices();
+	
+	// update skeletal mesh LOD (cf. USkeletalMesh::CommitMeshDescription)
+	SkeletalMesh->CommitMeshDescription(FSkeletonModifierLocals::LODIndex, *MeshDescription);
+
+	SkeletalMesh->PostEditChange();
+
+	// update skeleton
+	USkeleton* Skeleton = SkeletalMesh->GetSkeleton();
+	Skeleton->Modify();
+	if (Skeleton->RecreateBoneTree(SkeletalMesh))
 	{
-		FlushRenderingCommands();
-		SkeletalMesh->SetFlags(RF_Transactional);
-		SkeletalMesh->Modify();
-		// generate mesh import data from the mesh description (now with modified weights)
-		FSkeletalMeshImportData SkeletalMeshImportData = FSkeletalMeshImportData::CreateFromMeshDescription(*MeshDescription);
-		SkeletalMesh->SaveLODImportedData(0, SkeletalMeshImportData);
-		SkeletalMesh->SetLODImportedDataVersions(0, ESkeletalMeshGeoImportVersions::LatestVersion, ESkeletalMeshSkinningImportVersions::LatestVersion);
-		SkeletalMesh->SetUseLegacyMeshDerivedDataKey(false);
-		SkeletalMesh->PostEditChange();
+		Skeleton->MarkPackageDirty();	
 	}
 
 	return true;
diff --git a/Engine/Source/Runtime/Engine/Classes/Engine/SkeletalMesh.h b/Engine/Source/Runtime/Engine/Classes/Engine/SkeletalMesh.h
--- a/Engine/Source/Runtime/Engine/Classes/Engine/SkeletalMesh.h
+++ b/Engine/Source/Runtime/Engine/Classes/Engine/SkeletalMesh.h
@@ -68,6 +68,7 @@ struct FSkinWeightProfileInfo;
 class FSkinWeightVertexBuffer;
 enum class ESkeletalMeshGeoImportVersions : uint8;
 enum class ESkeletalMeshSkinningImportVersions : uint8;
+struct FMeshDescription;
 
 /*-----------------------------------------------------------------------------
 	Async Skeletal Mesh Compilation
@@ -538,6 +539,23 @@ public:
 	// End USkeletalMeshEditorData public skeletalmesh API
 	//////////////////////////////////////////////////////////////////////////
 
+	/* Structure that defines parameters passed into the commit mesh description function */
+	struct FCommitMeshDescriptionParams
+	{
+		FCommitMeshDescriptionParams()
+			: bModify(true)
+		{}
+
+		bool bModify;
+	};
+
+	/* Fills OutMeshDescription from the associated InLODIndex. */
+	void GetMeshDescription(const int32 InLODIndex, FMeshDescription& OutMeshDescription) const;
+
+	/* Sets InMeshDescription to the associated InLODIndex. */
+	void CommitMeshDescription(const int32 InLODIndex, const FMeshDescription& InMeshDescription,
+		const FCommitMeshDescriptionParams& InParams = FCommitMeshDescriptionParams());
+
 	/** Get the number of imported vertices. This returns 0 if GetImportedModel() returns a nullptr.
 	  * This is the number of vertices as they appear in the source asset, for example 8 for a cube. */
 	int32 GetNumImportedVertices() const;
diff --git a/Engine/Source/Runtime/Engine/Private/SkeletalMesh.cpp b/Engine/Source/Runtime/Engine/Private/SkeletalMesh.cpp
--- a/Engine/Source/Runtime/Engine/Private/SkeletalMesh.cpp
+++ b/Engine/Source/Runtime/Engine/Private/SkeletalMesh.cpp
@@ -1680,7 +1680,7 @@ void USkeletalMesh::Serialize( FArchive& Ar )
 	}
 
 	// make sure we're counting properly
-	if (!Ar.IsLoading() && !Ar.IsSaving())
+	if ((!Ar.IsLoading() && !Ar.IsSaving()) || Ar.IsTransacting())
 	{
 		Ar << GetRefBasesInvMatrix();
 	}
@@ -2400,6 +2400,46 @@ void USkeletalMesh::EmptyAllImportData()
 	}
 }
 
+void USkeletalMesh::GetMeshDescription(const int32 InLODIndex, FMeshDescription& OutMeshDescription) const
+{
+	// Check first if we have bulk data available and non-empty.
+	if (IsLODImportedDataBuildAvailable(InLODIndex) && !IsLODImportedDataEmpty(InLODIndex))
+	{
+		FSkeletalMeshImportData SkeletalMeshImportData;
+		LoadLODImportedData(InLODIndex, SkeletalMeshImportData);
+		SkeletalMeshImportData.GetMeshDescription(OutMeshDescription);
+	}
+	else
+	{
+		// Fall back on the LOD model directly if no bulk data exists. When we commit
+		// the mesh description, we override using the bulk data. This can happen for older
+		// skeletal meshes, from UE 4.24 and earlier.
+		const FSkeletalMeshModel* SkeletalMeshModel = GetImportedModel();
+		if (SkeletalMeshModel && SkeletalMeshModel->LODModels.IsValidIndex(InLODIndex))
+		{
+			SkeletalMeshModel->LODModels[InLODIndex].GetMeshDescription(OutMeshDescription, this);
+		}			
+	}
+}
+
+void USkeletalMesh::CommitMeshDescription(
+	const int32 InLODIndex, const FMeshDescription& InMeshDescription, const FCommitMeshDescriptionParams& InParams)
+{	
+	FSkeletalMeshImportData SkeletalMeshImportData = FSkeletalMeshImportData::CreateFromMeshDescription(InMeshDescription);
+
+	if (InParams.bModify)
+	{
+		SetFlags(RF_Transactional);
+		Modify();
+		GetMeshEditorData().Modify();
+	}
+
+	SaveLODImportedData(InLODIndex, SkeletalMeshImportData);
+	// Make sure the mesh builder knows it's the latest variety, so that the render data gets properly rebuilt.
+	SetLODImportedDataVersions(InLODIndex, ESkeletalMeshGeoImportVersions::LatestVersion, ESkeletalMeshSkinningImportVersions::LatestVersion);
+	SetUseLegacyMeshDerivedDataKey(false);
+}
+
 void USkeletalMesh::CreateUserSectionsDataForLegacyAssets()
 {
 	//We want to avoid changing the ddc if we load an old asset.
diff --git a/Engine/Source/Runtime/SkeletalMeshDescription/Public/SkinWeightsAttributesRef.h b/Engine/Source/Runtime/SkeletalMeshDescription/Public/SkinWeightsAttributesRef.h
--- a/Engine/Source/Runtime/SkeletalMeshDescription/Public/SkinWeightsAttributesRef.h
+++ b/Engine/Source/Runtime/SkeletalMeshDescription/Public/SkinWeightsAttributesRef.h
@@ -108,19 +108,9 @@ namespace Internal
 		class const_iterator
 		{
 		public:
-			const_iterator() = default;
-
 			const_iterator& operator++()
 			{
-				if (Index != INDEX_NONE && Weights)
-				{
-					++Index;
-					if (Index == Weights->Num())
-					{
-						Weights = nullptr;
-						Index = INDEX_NONE;
-					}
-				}
+				++Index;
 				return *this;
 			}
 
@@ -137,18 +127,22 @@ namespace Internal
 		protected:
 			friend class TVertexBoneWeightsBase<AdaptorType, StorageType>;
 
-			explicit const_iterator(const TVertexBoneWeightsBase<AdaptorType, StorageType> *InWeights) :
+			const_iterator() = delete;
+			
+			explicit const_iterator(const TVertexBoneWeightsBase<AdaptorType, StorageType> *InWeights, int32 InIndex) :
 	            Weights(InWeights),
-	            Index(0)
+	            Index(InIndex)
 			{
+				check(InWeights);
+				check(InIndex > INDEX_NONE && InIndex <= InWeights->Num());
 			}
 
 			const TVertexBoneWeightsBase<AdaptorType, StorageType> *Weights = nullptr;
 			int32 Index = INDEX_NONE;
 		};
 
-		const_iterator begin() const { return const_iterator(this); }		
-		const_iterator end() const { return const_iterator(); }
+		const_iterator begin() const { return const_iterator(this, 0); }		
+		const_iterator end() const { return const_iterator(this, this->Num()); }
 
 	private:
 		TArrayAttribute<StorageType> Container;
