Skelmesh Converter: Check if the mesh geometry has vertex color data, in which case set the flag on the skelmesh so that it gets displayed properly.

#rnx

diff --git a/Engine/Source/Developer/SkeletalMeshUtilitiesCommon/Private/StaticToSkeletalMeshConverter.cpp b/Engine/Source/Developer/SkeletalMeshUtilitiesCommon/Private/StaticToSkeletalMeshConverter.cpp
--- a/Engine/Source/Developer/SkeletalMeshUtilitiesCommon/Private/StaticToSkeletalMeshConverter.cpp
+++ b/Engine/Source/Developer/SkeletalMeshUtilitiesCommon/Private/StaticToSkeletalMeshConverter.cpp
@@ -12,6 +12,7 @@
 #include "LODUtilities.h"
 #include "MeshDescription.h"
 #include "MeshUtilities.h"
+#include "Misc/Guid.h"
 #include "Modules/ModuleManager.h"
 #include "ReferenceSkeleton.h"
 #include "Rendering/SkeletalMeshLODImporterData.h"
@@ -287,6 +288,27 @@ static bool AddLODFromStaticMeshSourceModel(
 }
 
 
+static bool HasVertexColors(
+	const USkeletalMesh* InSkeletalMesh
+	)
+{
+	for (const FSkeletalMeshLODModel& LODModel: InSkeletalMesh->GetImportedModel()->LODModels)
+	{
+		for (const FSkelMeshSection& Section: LODModel.Sections)
+		{
+			for (const FSoftSkinVertex& Vertex: Section.SoftVertices)
+			{
+				if (Vertex.Color != FColor::White)
+				{
+					return true;
+				}
+			}
+		}
+	}
+	return false;
+}
+
+
 bool FStaticToSkeletalMeshConverter::InitializeSkeletalMeshFromStaticMesh(
 	USkeletalMesh* InSkeletalMesh,
 	const UStaticMesh* InStaticMesh,
@@ -362,6 +384,12 @@ bool FStaticToSkeletalMeshConverter::InitializeSkeletalMeshFromStaticMesh(
 	}
 
 	InSkeletalMesh->SetMaterials(Materials);
+
+	if (HasVertexColors(InSkeletalMesh))
+	{
+		InSkeletalMesh->SetHasVertexColors(true);	
+		InSkeletalMesh->SetVertexColorGuid(FGuid::NewGuid());
+	}
 	
 	// Set the bounds from the static mesh, including the extensions, otherwise it won't render properly (among other things).
 	InSkeletalMesh->SetImportedBounds( InStaticMesh->GetBounds() );
@@ -510,6 +538,12 @@ bool FStaticToSkeletalMeshConverter::InitializeSkeletalMeshFromMeshDescriptions(
 		InSkeletalMesh->SetMaterials(NewMaterials);
 	}
 
+	if (HasVertexColors(InSkeletalMesh))
+	{
+		InSkeletalMesh->SetHasVertexColors(true);	
+		InSkeletalMesh->SetVertexColorGuid(FGuid::NewGuid());
+	}
+	
 	InSkeletalMesh->SetImportedBounds( FBox3d{BoundingBox} );
 
 	return true;
