Skelmesh Converter: Fix setting multiple LODs

#rnx

diff --git a/Engine/Plugins/Experimental/GeometryScripting/Source/GeometryScriptingCore/Private/MeshAssetFunctions.cpp b/Engine/Plugins/Experimental/GeometryScripting/Source/GeometryScriptingCore/Private/MeshAssetFunctions.cpp
--- a/Engine/Plugins/Experimental/GeometryScripting/Source/GeometryScriptingCore/Private/MeshAssetFunctions.cpp
+++ b/Engine/Plugins/Experimental/GeometryScripting/Source/GeometryScriptingCore/Private/MeshAssetFunctions.cpp
@@ -600,9 +600,25 @@ UDynamicMesh* UGeometryScriptLibrary_StaticMeshFunctions::CopyMeshToSkeletalMesh
 		Converter.Convert(&ReadMesh, MeshDescription, !Options.bEnableRecomputeTangents);
 	});
 
+	// Ensure we have enough LODInfos to cover up to the requested LOD.
+	for (int32 LODIndex = ToSkeletalMeshAsset->GetLODInfoArray().Num(); LODIndex <= TargetLOD.LODIndex; LODIndex++)
+	{
+		FSkeletalMeshLODInfo& LODInfo = ToSkeletalMeshAsset->AddLODInfo();
+		
+		ToSkeletalMeshAsset->GetImportedModel()->LODModels.Add(new FSkeletalMeshLODModel);
+		LODInfo.ReductionSettings.BaseLOD = 0;
+	}
+
 	FSkeletalMeshLODInfo* SkeletalLODInfo = ToSkeletalMeshAsset->GetLODInfo(TargetLOD.LODIndex);
 	SkeletalLODInfo->BuildSettings.bRecomputeNormals = Options.bEnableRecomputeNormals;
 	SkeletalLODInfo->BuildSettings.bRecomputeTangents = Options.bEnableRecomputeTangents;
+	
+	// Prevent decimation of this LOD.
+	SkeletalLODInfo->ReductionSettings.NumOfTrianglesPercentage = 1.0;
+	SkeletalLODInfo->ReductionSettings.NumOfVertPercentage = 1.0;
+	SkeletalLODInfo->ReductionSettings.MaxNumOfTriangles = MAX_int32;
+	SkeletalLODInfo->ReductionSettings.MaxNumOfVerts = MAX_int32; 
+	SkeletalLODInfo->ReductionSettings.BaseLOD = TargetLOD.LODIndex;
 
 	FSkeletalMeshImportData SkeletalMeshImportData = 
 		FSkeletalMeshImportData::CreateFromMeshDescription(MeshDescription);
diff --git a/Engine/Source/Developer/SkeletalMeshUtilitiesCommon/Private/StaticToSkeletalMeshConverter.cpp b/Engine/Source/Developer/SkeletalMeshUtilitiesCommon/Private/StaticToSkeletalMeshConverter.cpp
--- a/Engine/Source/Developer/SkeletalMeshUtilitiesCommon/Private/StaticToSkeletalMeshConverter.cpp
+++ b/Engine/Source/Developer/SkeletalMeshUtilitiesCommon/Private/StaticToSkeletalMeshConverter.cpp
@@ -246,6 +246,9 @@ static bool AddLODFromStaticMeshSourceModel(
 	CopyBuildSettings(InStaticMeshSourceModel.BuildSettings, SkeletalLODInfo.BuildSettings);
 	CopyReductionSettings(InStaticMeshSourceModel.ReductionSettings, SkeletalLODInfo.ReductionSettings);
 
+	FSkeletalMeshModel* ImportedModels = InSkeletalMesh->GetImportedModel();
+	const int32 LODIndex = ImportedModels->LODModels.Num(); 
+
 	if (InStaticMeshSourceModel.IsMeshDescriptionValid())
 	{
 		FMeshDescription SkeletalMeshGeometry;
@@ -256,6 +259,9 @@ static bool AddLODFromStaticMeshSourceModel(
 		
 		FSkeletalMeshAttributes SkeletalMeshAttributes(SkeletalMeshGeometry);
 		SkeletalMeshAttributes.Register();
+		
+		// Ensure that we don't build this LOD from an earlier LOD, but use the actual model given.
+		SkeletalLODInfo.ReductionSettings.BaseLOD = LODIndex;
 
 		// Full binding to the root bone.
 		FSkinWeightsVertexAttributesRef SkinWeights = SkeletalMeshAttributes.GetVertexSkinWeights();
@@ -274,9 +280,10 @@ static bool AddLODFromStaticMeshSourceModel(
 	}
 	else
 	{
-		FSkeletalMeshModel* ImportedModels = InSkeletalMesh->GetImportedModel();
-		const int32 LODIndex = ImportedModels->LODModels.Num(); 
 		ImportedModels->LODModels.Add(new FSkeletalMeshLODModel);
+		
+		// Base this LOD off of the previous one.
+		SkeletalLODInfo.ReductionSettings.BaseLOD = FMath::Max(0, LODIndex - 1); 
 			
 		FSkeletalMeshUpdateContext UpdateContext;
 		UpdateContext.SkeletalMesh = InSkeletalMesh;
