Fix material rematch when re-importing skeletalmesh

#rnx

diff --git a/Engine/Source/Developer/SkeletalMeshUtilitiesCommon/Private/LODUtilities.cpp b/Engine/Source/Developer/SkeletalMeshUtilitiesCommon/Private/LODUtilities.cpp
--- a/Engine/Source/Developer/SkeletalMeshUtilitiesCommon/Private/LODUtilities.cpp
+++ b/Engine/Source/Developer/SkeletalMeshUtilitiesCommon/Private/LODUtilities.cpp
@@ -3874,12 +3874,12 @@ void FLODUtilities::ReorderMaterialSlotToBaseLod(USkeletalMesh* SkeletalMesh)
 	//Make sure we have valid parameters
 	if (!ensure(SkeletalMesh))
 	{
-		UE_LOG(LogLODUtilities, Warning, TEXT("FLODUtilities::ReorderMaterialSlotToBaseLod: Bad parameters, SkeletalMesh is null"));
+		UE_ASSET_LOG(LogLODUtilities, Warning, SkeletalMesh, TEXT("FLODUtilities::ReorderMaterialSlotToBaseLod: Bad parameters, SkeletalMesh is null"));
 		return;
 	}
 	if (SkeletalMesh->IsLODImportedDataEmpty(0))
 	{
-		UE_LOG(LogLODUtilities, Warning, TEXT("FLODUtilities::ReorderMaterialSlotToBaseLod: Skeletal mesh invalid import data for LOD 0"));
+		UE_ASSET_LOG(LogLODUtilities, Warning, SkeletalMesh, TEXT("FLODUtilities::ReorderMaterialSlotToBaseLod: Skeletal mesh invalid import data for LOD 0"));
 		return;
 	}
 	FSkeletalMeshImportData BaseLodImportData;
@@ -3902,7 +3902,10 @@ void FLODUtilities::ReorderMaterialSlotToBaseLod(USkeletalMesh* SkeletalMesh)
 			if (MaterialSlotName == ImportedMaterialSlotName)
 			{
 				MaterialSlotRemap[MaterialIndex] = ImportedMaterialIndex;
-				AvailableIndexes[MaterialIndex] = false;
+				if(ensure(AvailableIndexes.IsValidIndex(ImportedMaterialIndex) && AvailableIndexes[ImportedMaterialIndex]))
+				{
+					AvailableIndexes[ImportedMaterialIndex] = false;
+				}
 				break;
 			}
 		}
@@ -3971,12 +3974,15 @@ void FLODUtilities::ReorderMaterialSlotToBaseLod(USkeletalMesh* SkeletalMesh)
 			for (int32 SectionIndex = 0; SectionIndex < LODModel.Sections.Num(); ++SectionIndex)
 			{
 				FSkelMeshSection& Section = LODModel.Sections[SectionIndex];
-				Section.MaterialIndex = MaterialSlotRemap[Section.MaterialIndex];
+				if(MaterialSlotRemap.IsValidIndex(Section.MaterialIndex))
+				{
+					Section.MaterialIndex = MaterialSlotRemap[Section.MaterialIndex];
+				}
 			}
 		}
 		else
 		{
-			UE_LOG(LogLODUtilities, Warning, TEXT("FLODUtilities::ReorderMaterialSlotToBaseLod: Skeletal mesh invalid LODModel %d"), LodIndex);
+			UE_ASSET_LOG(LogLODUtilities, Warning, SkeletalMesh, TEXT("FLODUtilities::ReorderMaterialSlotToBaseLod: Skeletal mesh invalid LODModel %d"), LodIndex);
 		}
 		
 		//Remap the LODMaterialMap
@@ -3992,7 +3998,7 @@ void FLODUtilities::ReorderMaterialSlotToBaseLod(USkeletalMesh* SkeletalMesh)
 		}
 		else
 		{
-			UE_LOG(LogLODUtilities, Warning, TEXT("FLODUtilities::ReorderMaterialSlotToBaseLod: Skeletal mesh invalid LODInfo %d"), LodIndex);
+			UE_ASSET_LOG(LogLODUtilities, Warning, SkeletalMesh, TEXT("FLODUtilities::ReorderMaterialSlotToBaseLod: Skeletal mesh invalid LODInfo %d"), LodIndex);
 		}
 	}
 }
