Skelmesh: Optimize import of alternative skin weights.

diff --git a/Engine/Source/Developer/MeshUtilities/Private/MeshUtilities.cpp b/Engine/Source/Developer/MeshUtilities/Private/MeshUtilities.cpp
--- a/Engine/Source/Developer/MeshUtilities/Private/MeshUtilities.cpp
+++ b/Engine/Source/Developer/MeshUtilities/Private/MeshUtilities.cpp
@@ -5775,6 +5775,56 @@ void FMeshUtilities::CalculateOverlappingCorners(const TArray<FVector3f>& InVert
 }
 
 
+struct FHashableSkinWeight
+{
+	FHashableSkinWeight() = default;
+	FHashableSkinWeight(const FHashableSkinWeight&) = default;
+	FHashableSkinWeight(FHashableSkinWeight&&) = default;
+	FHashableSkinWeight(const FRawSkinWeight& InRawWeight, const int32 InNumInfluences) :
+		NumInfluences(InNumInfluences),
+		InfluenceBones(InRawWeight.InfluenceBones),
+		InfluenceWeights(InRawWeight.InfluenceWeights)
+	{
+	}
+	
+	const int32 NumInfluences = 0;
+	const FBoneIndexType *InfluenceBones = nullptr;
+	const uint16 *InfluenceWeights = nullptr;
+};
+
+struct FHashableSkinWeightKeyFuncs : BaseKeyFuncs<TPair<FHashableSkinWeight, int32>, FHashableSkinWeight>
+{
+	static const FHashableSkinWeight& GetSetKey(const TPair<FHashableSkinWeight, int32>& Element)
+	{
+		return Element.Key;
+	}
+
+	static bool Matches(const FHashableSkinWeight& A, const FHashableSkinWeight& B)
+	{
+		for (int32 Index = 0; Index < A.NumInfluences; Index++)
+		{
+			if (A.InfluenceBones[Index] != B.InfluenceBones[Index] ||
+				A.InfluenceWeights[Index] != B.InfluenceWeights[Index])
+			{
+				return false;
+			}
+		}
+		return true;
+	}
+	
+	static uint32 GetKeyHash(const FHashableSkinWeight& Key)
+	{
+		uint32 Hash = GetTypeHash(Key.NumInfluences);
+		for (int32 Index = 0; Index < Key.NumInfluences; Index++)
+		{
+			Hash = HashCombineFast(Hash, GetTypeHash(Key.InfluenceBones[Index]));
+			Hash = HashCombineFast(Hash, GetTypeHash(Key.InfluenceWeights[Index]));
+		}
+		return Hash;
+	}
+};
+
+
 void FMeshUtilities::GenerateRuntimeSkinWeightData(
 	const FSkeletalMeshLODModel* ImportedModel,
 	const TArray<FRawSkinWeight>& InRawSkinWeights,
@@ -5797,7 +5847,9 @@ void FMeshUtilities::GenerateRuntimeSkinWeightData(
 		const bool b16BitBoneIndices = TargetLODModel.DoSectionsUse16BitBoneIndex();
 		InOutSkinWeightOverrideData.b16BitBoneIndices = b16BitBoneIndices;
 
-		TArray<FRawSkinWeight> UniqueWeights;
+		TMap<FHashableSkinWeight, int32, FDefaultSetAllocator, FHashableSkinWeightKeyFuncs> UniqueWeights;
+		int32 OverrideCount = 0;
+		
 		for (int32 VertexIndex = 0; VertexIndex < TargetVertices.Num(); ++VertexIndex)
 		{
 			// Take each original skin weight from the LOD and compare it with supplied alternative weight data
@@ -5817,24 +5869,12 @@ void FMeshUtilities::GenerateRuntimeSkinWeightData(
 
 			if (bIsDifferent)
 			{
-				// Check whether or not there is already an override store which matches the new skin weight data
-				int32 OverrideIndex = UniqueWeights.IndexOfByPredicate([SourceSkinWeight, NumInfluences](const FRawSkinWeight& InOverride)
-				{
-					bool bSame = true;
-					for (int32 InfluenceIndex = 0; InfluenceIndex < NumInfluences; ++InfluenceIndex)
-					{
-						bSame &= (InOverride.InfluenceBones[InfluenceIndex] == SourceSkinWeight.InfluenceBones[InfluenceIndex]);
-						bSame &= (InOverride.InfluenceWeights[InfluenceIndex] == SourceSkinWeight.InfluenceWeights[InfluenceIndex]);
-					}
-
-					return bSame;
-				});
-
+				const int32* OverrideIndexPtr = UniqueWeights.Find({SourceSkinWeight, NumInfluences});
+				
 				// If one hasn't been added yet, create a new one
-				if (OverrideIndex == INDEX_NONE)
+				int32 OverrideIndex;
+				if (OverrideIndexPtr == nullptr)
 				{
-					OverrideIndex = UniqueWeights.Num();
-
 					// When writing out 8-bit weights, we want to collect the resulting deviation from fully normalized
 					// weight sum. Once all values are copied, we alter the first, and most prominent, weight, by the 
 					// deviation value so that we end up with completely normalized 8-bit weights.
@@ -5871,13 +5911,19 @@ void FMeshUtilities::GenerateRuntimeSkinWeightData(
 						}
 					}
 
-					UniqueWeights.Add(SourceSkinWeight);
+					OverrideIndex = OverrideCount;
+					UniqueWeights.Add({SourceSkinWeight, NumInfluences}, OverrideCount);
+					OverrideCount++;
 
 					if (!bInUseHighPrecisionWeights)
 					{
 						InOutSkinWeightOverrideData.BoneWeights[FirstWeightOffset] += Accumulated8BitDeviation;
 					}
 				}
+				else
+				{
+					OverrideIndex = *OverrideIndexPtr;
+				}
 
 				InOutSkinWeightOverrideData.VertexIndexToInfluenceOffset.Add(VertexIndex, OverrideIndex);
 			}
diff --git a/Engine/Source/Developer/SkeletalMeshUtilitiesCommon/Private/LODUtilities.cpp b/Engine/Source/Developer/SkeletalMeshUtilitiesCommon/Private/LODUtilities.cpp
--- a/Engine/Source/Developer/SkeletalMeshUtilitiesCommon/Private/LODUtilities.cpp
+++ b/Engine/Source/Developer/SkeletalMeshUtilitiesCommon/Private/LODUtilities.cpp
@@ -2245,18 +2245,24 @@ bool FLODUtilities::UpdateAlternateSkinWeights(
 	VertexIndexSrcToVertexIndexDestMatches.Reserve(VertexNumberSrc);
 	TArray<uint32> VertexIndexToMatchWithPositions;
 
-	auto FindWedgeIndexesUsingVertexIndex = [](const FSkeletalMeshImportData& ImportData, const int32 VertexIndex, TArray<int32>& OutWedgeIndexes)
+	TMultiMap<int32, int32> VertexToSrcWedgeMap;
+	TMultiMap<int32, int32> VertexToDestWedgeMap;
+
+	for (int32 WedgeIndex = 0; WedgeIndex < ImportDataSrc.Wedges.Num(); WedgeIndex++)
 	{
-		for (int32 WedgeIndex = 0; WedgeIndex < ImportData.Wedges.Num(); ++WedgeIndex)
-		{
-			const SkeletalMeshImportData::FVertex& Wedge = ImportData.Wedges[WedgeIndex];
-			if (Wedge.VertexIndex == VertexIndex)
-			{
-				OutWedgeIndexes.Add(WedgeIndex);
-			}
-		}
-	};
+		const SkeletalMeshImportData::FVertex& Wedge = ImportDataSrc.Wedges[WedgeIndex];
+		VertexToSrcWedgeMap.Add(Wedge.VertexIndex, WedgeIndex);
+	}
+	for (int32 WedgeIndex = 0; WedgeIndex < ImportDataDest.Wedges.Num(); WedgeIndex++)
+	{
+		const SkeletalMeshImportData::FVertex& Wedge = ImportDataDest.Wedges[WedgeIndex];
+		VertexToDestWedgeMap.Add(Wedge.VertexIndex, WedgeIndex);
+	}
 
+	TArray<int32> SrcWedgeIndexes;
+	TArray<int32> DestWedgeIndexes;
+
+	
 	// Match all source vertex with destination vertex
 	for (int32 VertexIndexSrc = 0; VertexIndexSrc < PointNumberSrc; ++VertexIndexSrc)
 	{
@@ -2275,8 +2281,8 @@ bool FLODUtilities::UpdateAlternateSkinWeights(
 			//We have a direct match
 			VertexMatchNameSpace::FVertexMatchResult& VertexMatchDest = VertexIndexSrcToVertexIndexDestMatches.Add(VertexIndexSrc);
 
-			TArray<int32> SrcWedgeIndexes;
-			FindWedgeIndexesUsingVertexIndex(ImportDataSrc, VertexIndexSrc, SrcWedgeIndexes);
+			SrcWedgeIndexes.Reset();
+			VertexToSrcWedgeMap.MultiFind(VertexIndexSrc, SrcWedgeIndexes);
 
 			if (SrcWedgeIndexes.Num() > 0 && SimilarDestinationVertex.Num() > 1)
 			{
@@ -2284,8 +2290,9 @@ bool FLODUtilities::UpdateAlternateSkinWeights(
 				for (int32 MatchDestinationIndex = 0; MatchDestinationIndex < SimilarDestinationVertex.Num(); ++MatchDestinationIndex)
 				{
 					int32 VertexIndexDest = SimilarDestinationVertex[MatchDestinationIndex];
-					TArray<int32> DestWedgeIndexes;
-					FindWedgeIndexesUsingVertexIndex(ImportDataDest, VertexIndexDest, DestWedgeIndexes);
+					DestWedgeIndexes.Reset();
+					VertexToDestWedgeMap.MultiFind(VertexIndexDest, DestWedgeIndexes);
+					
 					for (int32 IndexDest = 0; IndexDest < DestWedgeIndexes.Num(); ++IndexDest)
 					{
 						int32 DestWedgeIndex = DestWedgeIndexes[IndexDest];
@@ -2669,6 +2676,14 @@ void FLODUtilities::GenerateImportedSkinWeightProfileData(
 
 	int32 MaxNumInfluences = 0;
 
+	TMultiMap<int32, const SkeletalMeshImportData::FVertInfluence*> VertexToInfluenceMap;
+	for (const SkeletalMeshImportData::FVertInfluence& VertInfluence : ImportedProfileData.SourceModelInfluences)
+	{
+		VertexToInfluenceMap.Add(VertInfluence.VertIndex, &VertInfluence);
+	}
+	TMap<FBoneIndexType, float> WeightForBone;
+	TArray<const SkeletalMeshImportData::FVertInfluence*> FoundInfluences;
+
 	for (int32 VertexInstanceIndex = 0; VertexInstanceIndex < DestinationSoftVertices.Num(); ++VertexInstanceIndex)
 	{
 		int32 SectionIndex = INDEX_NONE;
@@ -2685,21 +2700,22 @@ void FLODUtilities::GenerateImportedSkinWeightProfileData(
 		check(VertexIndex >= 0 && VertexIndex <= LODModelDest.MaxImportVertex);
 		FRawSkinWeight& SkinWeight = SkinWeights.AddDefaulted_GetRef();
 
-		TMap<FBoneIndexType, float> WeightForBone;
-		for (const SkeletalMeshImportData::FVertInfluence& VertInfluence : ImportedProfileData.SourceModelInfluences)
+		WeightForBone.Reset();
+		FoundInfluences.Reset();
+		
+		VertexToInfluenceMap.MultiFind(VertexIndex, FoundInfluences);
+		
+		for (const SkeletalMeshImportData::FVertInfluence* VertInfluence : FoundInfluences)
 		{
-			if(VertexIndex == VertInfluence.VertIndex)
+			//Use the section bone map to remap the bone index
+			int32 BoneMapIndex = INDEX_NONE;
+			SectionBoneMap.Find(VertInfluence->BoneIndex, BoneMapIndex);
+			if (BoneMapIndex == INDEX_NONE)
 			{
-				//Use the section bone map to remap the bone index
-				int32 BoneMapIndex = INDEX_NONE;
-				SectionBoneMap.Find(VertInfluence.BoneIndex, BoneMapIndex);
-				if (BoneMapIndex == INDEX_NONE)
-				{
-					//Map to root of the section
-					BoneMapIndex = 0;
-				}
-				WeightForBone.Add(static_cast<FBoneIndexType>(BoneMapIndex), VertInfluence.Weight);
+				//Map to root of the section
+				BoneMapIndex = 0;
 			}
+			WeightForBone.Add(static_cast<FBoneIndexType>(BoneMapIndex), VertInfluence->Weight);
 		}
 
 		using namespace UE::AnimationCore;
@@ -2710,7 +2726,7 @@ void FLODUtilities::GenerateImportedSkinWeightProfileData(
 		FBoneIndexType InfluenceBones[MAX_TOTAL_INFLUENCES];
 		float InfluenceWeights[MAX_TOTAL_INFLUENCES];
 		
-		for (auto Kvp : WeightForBone)
+		for (const TTuple<FBoneIndexType, float>& Kvp : WeightForBone)
 		{
 			InfluenceBones[InfluenceBoneCount] = Kvp.Key;
 			InfluenceWeights[InfluenceBoneCount] = Kvp.Value;
