Add extra info in message when skin weight profiles don't match

diff --git a/Engine/Source/Developer/SkeletalMeshUtilitiesCommon/Private/LODUtilities.cpp b/Engine/Source/Developer/SkeletalMeshUtilitiesCommon/Private/LODUtilities.cpp
--- a/Engine/Source/Developer/SkeletalMeshUtilitiesCommon/Private/LODUtilities.cpp
+++ b/Engine/Source/Developer/SkeletalMeshUtilitiesCommon/Private/LODUtilities.cpp
@@ -2167,8 +2167,12 @@ bool FLODUtilities::UpdateAlternateSkinWeights(
 	{
 		FFormatNamedArguments Args;
 		Args.Add(TEXT("SkeletalMeshName"), FText::FromString(SkeletalMeshDest->GetName()));
+		Args.Add(TEXT("ProfileName"), FText::FromName(ProfileNameDest));
+		Args.Add(TEXT("LODIndex"), LODIndexDest);
 
-		FText Message = FText::Format(NSLOCTEXT("FLODUtilities_UpdateAlternateSkinWeights", "AlternateDataNotAvailable", "Asset {SkeletalMeshName} failed to import skin weight profile the alternate skinning imported source data is not available."), Args);
+		FText Message = FText::Format(NSLOCTEXT("FLODUtilities_UpdateAlternateSkinWeights", "AlternateDataNotAvailable", 
+			"Asset {SkeletalMeshName} LOD {LODIndex} failed to import skin weight profile {ProfileName}. The alternate skinning imported source data is not available."), 
+			Args);
 		UE_LOG(LogLODUtilities, Warning, TEXT("%s"), *(Message.ToString()));
 		return false;
 	}
