LODUtilities MatchImportedMaterials fix for skeletal with no material re-imports

diff --git a/Engine/Source/Developer/SkeletalMeshUtilitiesCommon/Private/LODUtilities.cpp b/Engine/Source/Developer/SkeletalMeshUtilitiesCommon/Private/LODUtilities.cpp
--- a/Engine/Source/Developer/SkeletalMeshUtilitiesCommon/Private/LODUtilities.cpp
+++ b/Engine/Source/Developer/SkeletalMeshUtilitiesCommon/Private/LODUtilities.cpp
@@ -3648,6 +3648,11 @@ void FLODUtilities::MatchImportedMaterials(FLODUtilities::FSkeletalMeshMatchImpo
 			UE_ASSET_LOG(LogLODUtilities, Warning, Parameters.SkeletalMesh, TEXT("FLODUtilities::MatchImportedMaterials: Bad parameters, ExistingOriginalPerSectionMaterialImportName is null"));
 			return;
 		}
+
+		if (Parameters.SkeletalMesh->GetMaterials().Num() == 0)
+		{
+			return;
+		}
 	}
 	FSkeletalMeshLODInfo* LodInfo = Parameters.SkeletalMesh->GetLODInfo(Parameters.LodIndex);
 	if (!LodInfo)
@@ -3672,7 +3677,11 @@ void FLODUtilities::MatchImportedMaterials(FLODUtilities::FSkeletalMeshMatchImpo
 			, const TArray<FName>& PerSectionMaterialImportName)->FName
 		{
 			const TArray<FSkeletalMaterial>& MeshMaterials = Parameters.SkeletalMesh->GetMaterials();
-			check(MeshMaterials.Num() > 0);
+			if (!ensure(MeshMaterials.Num() > 0))
+			{
+				OutMaterialIndex = -1;
+				return NAME_None;
+			}
 			if (PerSectionMaterialImportName.IsValidIndex(SectionIndex))
 			{
 				for (int32 MaterialIndex = 0; MaterialIndex < MeshMaterials.Num(); ++MaterialIndex)
@@ -3758,7 +3767,8 @@ void FLODUtilities::MatchImportedMaterials(FLODUtilities::FSkeletalMeshMatchImpo
 						break;
 					}
 				}
-				else if (ImportedMaterials[NewSectionIndex].Material.Get() == Parameters.SkeletalMesh->GetMaterials()[ExistingCurrentMaterialIndex].MaterialInterface.Get()) //Use material slot compare to match in case the name is none
+				else if (Parameters.SkeletalMesh->GetMaterials().IsValidIndex(ExistingCurrentMaterialIndex) &&
+						 ImportedMaterials[NewSectionIndex].Material.Get() == Parameters.SkeletalMesh->GetMaterials()[ExistingCurrentMaterialIndex].MaterialInterface.Get()) //Use material slot compare to match in case the name is none
 				{
 					CopySectionData();
 					break;
@@ -3906,6 +3916,10 @@ void FLODUtilities::ReorderMaterialSlotToBaseLod(USkeletalMesh* SkeletalMesh)
 	SkeletalMesh->LoadLODImportedData(0, BaseLodImportData);
 	const TArray<SkeletalMeshImportData::FMaterial>& ImportedMaterials = BaseLodImportData.Materials;
 	TArray<FSkeletalMaterial>& Materials = SkeletalMesh->GetMaterials();
+	if (Materials.Num() == 0)
+	{
+		return;
+	}
 
 	TArray<int32> MaterialSlotRemap;
 	MaterialSlotRemap.AddUninitialized(Materials.Num());
