Fix a possible deadlock when building skeletalmesh

#rnx

diff --git a/Engine/Source/Developer/SkeletalMeshUtilitiesCommon/Private/LODUtilities.cpp b/Engine/Source/Developer/SkeletalMeshUtilitiesCommon/Private/LODUtilities.cpp
--- a/Engine/Source/Developer/SkeletalMeshUtilitiesCommon/Private/LODUtilities.cpp
+++ b/Engine/Source/Developer/SkeletalMeshUtilitiesCommon/Private/LODUtilities.cpp
@@ -3399,25 +3399,31 @@ void FLODUtilities::BuildMorphTargets(USkeletalMesh* BaseSkelMesh, FSkeletalMesh
 			{
 				if (LODIndex == 0)
 				{
-					// Required both for NewObject and to avoid a fatal error in StaticFindObject.
-					FGCScopeGuard GCScopeGuard;
-
-					if (!IsInGameThread())
+					//Garbage collect must be delayed until the skeletal mesh build is done by registering to delegate
+					//FCoreUObjectDelegates::GetPreGarbageCollectDelegate.
+					if (ensure(!GIsGarbageCollecting))
 					{
-						//TODO remove this code when overriding a UObject will be allow outside of the game thread
-						//We currently need to avoid overriding an existing asset outside of the game thread
-						UObject* ExistingMorphTarget = StaticFindObject(UMorphTarget::StaticClass(), BaseSkelMesh, *ShapeName);
-						if (ExistingMorphTarget)
+						if (!IsInGameThread())
 						{
-							//make sure the object is not standalone or transactional
-							ExistingMorphTarget->ClearFlags(RF_Standalone | RF_Transactional);
-							//Move this object in the transient package
-							ExistingMorphTarget->Rename(nullptr, GetTransientPackage(), REN_ForceNoResetLoaders | REN_DoNotDirty | REN_DontCreateRedirectors | REN_NonTransactional);
-							ExistingMorphTarget = nullptr;
+							//TODO remove this code when overriding a UObject will be allow outside of the game thread
+							//We currently need to avoid overriding an existing asset outside of the game thread
+							UObject* ExistingMorphTarget = StaticFindObject(UMorphTarget::StaticClass(), BaseSkelMesh, *ShapeName);
+							if (ExistingMorphTarget)
+							{
+								//make sure the object is not standalone or transactional
+								ExistingMorphTarget->ClearFlags(RF_Standalone | RF_Transactional);
+								//Move this object in the transient package
+								ExistingMorphTarget->Rename(nullptr, GetTransientPackage(), REN_ForceNoResetLoaders | REN_DoNotDirty | REN_DontCreateRedirectors | REN_NonTransactional);
+								ExistingMorphTarget = nullptr;
+							}
 						}
-					}
 
-					MorphTarget = NewObject<UMorphTarget>(BaseSkelMesh, ObjectName);
+						MorphTarget = NewObject<UMorphTarget>(BaseSkelMesh, ObjectName);
+					}
+					else
+					{
+						UE_ASSET_LOG(LogLODUtilities, Error, BaseSkelMesh, TEXT("FLODUtilities::BuildMorphTargets: Garbage collection is running during the skeletal build. Morph target [%s] cannot be built properly and will be missing."), *ShapeName);
+					}
 				}
 				else
 				{
