Skelmesh Conversion: Don't override the base LOD value from the static mesh's reduction settings.

#rnx

diff --git a/Engine/Source/Developer/SkeletalMeshUtilitiesCommon/Private/StaticToSkeletalMeshConverter.cpp b/Engine/Source/Developer/SkeletalMeshUtilitiesCommon/Private/StaticToSkeletalMeshConverter.cpp
--- a/Engine/Source/Developer/SkeletalMeshUtilitiesCommon/Private/StaticToSkeletalMeshConverter.cpp
+++ b/Engine/Source/Developer/SkeletalMeshUtilitiesCommon/Private/StaticToSkeletalMeshConverter.cpp
@@ -260,9 +260,6 @@ static bool AddLODFromStaticMeshSourceModel(
 		FSkeletalMeshAttributes SkeletalMeshAttributes(SkeletalMeshGeometry);
 		SkeletalMeshAttributes.Register();
 		
-		// Ensure that we don't build this LOD from an earlier LOD, but use the actual model given.
-		SkeletalLODInfo.ReductionSettings.BaseLOD = LODIndex;
-
 		// Full binding to the root bone.
 		FSkinWeightsVertexAttributesRef SkinWeights = SkeletalMeshAttributes.GetVertexSkinWeights();
 		UE::AnimationCore::FBoneWeight RootInfluence(InBoneIndex, 1.0f);
@@ -282,9 +279,6 @@ static bool AddLODFromStaticMeshSourceModel(
 	{
 		ImportedModels->LODModels.Add(new FSkeletalMeshLODModel);
 		
-		// Base this LOD off of the previous one.
-		SkeletalLODInfo.ReductionSettings.BaseLOD = FMath::Max(0, LODIndex - 1); 
-			
 		FSkeletalMeshUpdateContext UpdateContext;
 		UpdateContext.SkeletalMesh = InSkeletalMesh;
 		
