Skelmesh Conversion: Ensure build data is set to latest version, otherwise no build options show up.

#rnx

diff --git a/Engine/Source/Developer/SkeletalMeshUtilitiesCommon/Private/StaticToSkeletalMeshConverter.cpp b/Engine/Source/Developer/SkeletalMeshUtilitiesCommon/Private/StaticToSkeletalMeshConverter.cpp
--- a/Engine/Source/Developer/SkeletalMeshUtilitiesCommon/Private/StaticToSkeletalMeshConverter.cpp
+++ b/Engine/Source/Developer/SkeletalMeshUtilitiesCommon/Private/StaticToSkeletalMeshConverter.cpp
@@ -227,6 +227,7 @@ static bool AddLODFromMeshDescription(
 	}
 	
 	InSkeletalMesh->SaveLODImportedData(LODIndex, SkeletalMeshImportGeometry);
+	InSkeletalMesh->SetLODImportedDataVersions(LODIndex, ESkeletalMeshGeoImportVersions::LatestVersion, ESkeletalMeshSkinningImportVersions::LatestVersion);
 	
 	return true;
 }
