Initialize new USkeletalMesh assets with reduction settings that preserve the source mesh
- Resolves a bug where skeletal meshes were compressed to 50% of the source mesh
- These default settings were mirrored from InterchangeSkeletalMeshFactory

diff --git a/Engine/Source/Developer/SkeletalMeshUtilitiesCommon/Private/StaticToSkeletalMeshConverter.cpp b/Engine/Source/Developer/SkeletalMeshUtilitiesCommon/Private/StaticToSkeletalMeshConverter.cpp
--- a/Engine/Source/Developer/SkeletalMeshUtilitiesCommon/Private/StaticToSkeletalMeshConverter.cpp
+++ b/Engine/Source/Developer/SkeletalMeshUtilitiesCommon/Private/StaticToSkeletalMeshConverter.cpp
@@ -491,6 +491,10 @@ bool FStaticToSkeletalMeshConverter::InitializeSkeletalMeshFromMeshDescriptions(
 		{
 			// Add default LOD build settings.
 			FSkeletalMeshLODInfo& SkeletalLODInfo = InSkeletalMesh->AddLODInfo();
+			SkeletalLODInfo.ReductionSettings.NumOfTrianglesPercentage = 1.0f;
+			SkeletalLODInfo.ReductionSettings.NumOfVertPercentage = 1.0f;
+			SkeletalLODInfo.ReductionSettings.MaxDeviationPercentage = 0.0f;
+			SkeletalLODInfo.LODHysteresis = 0.02f;
 
 			SkeletalLODInfo.BuildSettings.bRecomputeNormals = bInRecomputeNormals;
 			SkeletalLODInfo.BuildSettings.bRecomputeTangents = bInRecomputeTangents;
