AltSkin Import: Perform a naive skin weight influence infill for vertexes we could not map to the base mesh, rather than just blindly mapping them to root.

diff --git a/Engine/Source/Developer/SkeletalMeshUtilitiesCommon/Private/LODUtilities.cpp b/Engine/Source/Developer/SkeletalMeshUtilitiesCommon/Private/LODUtilities.cpp
--- a/Engine/Source/Developer/SkeletalMeshUtilitiesCommon/Private/LODUtilities.cpp
+++ b/Engine/Source/Developer/SkeletalMeshUtilitiesCommon/Private/LODUtilities.cpp
@@ -4,6 +4,7 @@
 
 #if WITH_EDITOR
 
+#include "Algo/Accumulate.h"
 #include "Animation/MorphTarget.h"
 #include "Animation/SkinWeightProfile.h"
 #include "Async/ParallelFor.h"
@@ -19,7 +20,6 @@
 #include "Engine/SkinnedAssetAsyncCompileUtils.h"
 #include "Engine/SkinnedAssetCommon.h"
 #include "Engine/Texture2D.h"
-#include "Framework/Commands/UIAction.h"
 #include "Framework/Notifications/NotificationManager.h"
 #include "GenericQuadTree.h"
 #include "HAL/ThreadSafeBool.h"
@@ -31,7 +31,6 @@
 #include "Interfaces/ITargetPlatformManagerModule.h"
 #include "MeshUtilities.h"
 #include "MeshUtilitiesCommon.h"
-#include "Misc/CoreMisc.h"
 #include "Misc/FeedbackContext.h"
 #include "Misc/MessageDialog.h"
 #include "Modules/ModuleManager.h"
@@ -2152,6 +2151,158 @@ bool FLODUtilities::UpdateAlternateSkinWeights(
 	return UpdateAlternateSkinWeights(SkeletalMeshDest, ProfileNameDest, LODIndexDest, Options);
 }
 
+
+static void InfillMissingVertexSkinWeights(
+	const FSkeletalMeshImportData& InMesh,
+	const TArray<int32>& InMissingVertexes,
+	TArray<SkeletalMeshImportData::FRawBoneInfluence> &InOutInfluences
+	)
+{
+	// Create a map from influences to vertices.
+	TMultiMap<int32 /* Vertex */, int32 /* AlternateInfluencesIndex */> VertexToInfluenceMap;
+	VertexToInfluenceMap.Reserve(InOutInfluences.Num());
+	for (int32 InfluenceIndex = 0; InfluenceIndex < InOutInfluences.Num(); InfluenceIndex++)
+	{
+		VertexToInfluenceMap.Add(InOutInfluences[InfluenceIndex].VertexIndex, InfluenceIndex);
+	}
+
+	// Map from vertices to faces so we can track down neighbors. 
+	TMultiMap<int32 /* Vertex */, int32 /* Face */> VertexToFaceMap;
+	VertexToFaceMap.Reserve(InMesh.Wedges.Num() * 3);
+	for (int32 FaceIndex = 0; FaceIndex < InMesh.Faces.Num(); FaceIndex++)
+	{
+		for (int32 CornerIndex = 0; CornerIndex < 3; CornerIndex++)
+		{
+			const int32 WedgeIndex = InMesh.Faces[FaceIndex].WedgeIndex[CornerIndex];
+			VertexToFaceMap.Add(InMesh.Wedges[WedgeIndex].VertexIndex, FaceIndex);
+		}
+	}
+
+	// Compute neighborhood vertices for each unassigned point and assign them weights based on half the area of the triangle
+	// they belong to. We use this area for final weighting of influence contributions.
+	TArray<int32> Faces;
+	TMultiMap<int32 /* Vertex */, TPair<int32 /* NeighborVertex */, float /* NeighborArea */>> NeighborVertexesAndArea;
+
+	// We tally up the ratio of assigned influences for each unassigned vertex ranging from all neighbors have influences (1.0) to
+	// none have influences (0.0).
+	TArray<float> AssignedNeighborWeightRatio;
+	
+	for (int32 UnassignedVertex: InMissingVertexes)
+	{
+		Faces.Reset();
+		VertexToFaceMap.MultiFind(UnassignedVertex, Faces);
+		check(!Faces.IsEmpty());
+
+		TMap<int32, float> NeighborAndArea;
+
+		for (const int32 FaceIndex: Faces)
+		{
+			int32 VertexIndex[3];
+			
+			for (int32 CornerIndex = 0; CornerIndex < 3; CornerIndex++)
+			{
+				const int32 WedgeIndex  = InMesh.Faces[FaceIndex].WedgeIndex[CornerIndex];
+				VertexIndex[CornerIndex] = InMesh.Wedges[WedgeIndex].VertexIndex;
+			}
+
+			float Area = TriangleUtilities::ComputeTriangleArea(InMesh.Points[VertexIndex[0]], InMesh.Points[VertexIndex[1]], InMesh.Points[VertexIndex[2]]);
+			for (int32 CornerIndex = 0; CornerIndex < 3; CornerIndex++)
+			{
+				if (VertexIndex[CornerIndex] != UnassignedVertex)
+				{
+					NeighborAndArea.FindOrAdd(VertexIndex[CornerIndex], 0.0f) += Area / 2.0f;
+				}
+			}
+		}
+		int32 AssignedCount = 0;
+		for (const TPair<int32, float>& Item: NeighborAndArea)
+		{
+			NeighborVertexesAndArea.Add(UnassignedVertex, Item);
+
+			// Check if this neighbor has any influences computed already.
+			if (VertexToInfluenceMap.Contains(Item.Key))
+			{
+				AssignedCount++;
+			}
+		}
+		AssignedNeighborWeightRatio.Add(AssignedCount / static_cast<float>(NeighborAndArea.Num()));
+	}
+
+	// Order the unassigned vertices by how many counts are assigned. Start with the fully assigned ones first.
+	TSet<int32> UnassignedVertexQueue;
+	{
+		TArray<TPair<int32, float>> VertexWithRatio;
+		for (int32 Index = 0; Index < InMissingVertexes.Num(); Index++)
+		{
+			VertexWithRatio.Add({InMissingVertexes[Index], AssignedNeighborWeightRatio[Index]});
+		}
+		// Sort in descending order of ratio.
+		VertexWithRatio.Sort([](const TPair<int32, float>& A, const TPair<int32, float>& B) { return A.Value > B.Value; });
+
+		for (const TPair<int32, float>& Item: VertexWithRatio)
+		{
+			UnassignedVertexQueue.Add(Item.Key);
+		}
+	}
+
+	// Progressively try to resolve weights. If we end up with a set of points we can't resolve (e.g. islands, etc. we warn the user that some
+	// of their vertices will be mapped to root).
+	TArray<TPair<int32 /* NeighborVertex */, float /* NeighborArea */>> NeighborVertexAreaLookup;
+	TArray<int32> InfluenceIndexLookup;
+	while (!UnassignedVertexQueue.IsEmpty())
+	{
+		bool bPointRemoved = false;
+		for (int32 UnassignedVertex: UnassignedVertexQueue)
+		{
+			NeighborVertexAreaLookup.Reset();
+			NeighborVertexesAndArea.MultiFind(UnassignedVertex, NeighborVertexAreaLookup);
+
+			if (NeighborVertexAreaLookup.IsEmpty())
+			{
+				continue;
+			}
+
+			TMap</* BoneIndex */ int32, /* Weight */ float> VertexInfluences;
+			// Count up the influences that each neighbor has and weigh them by area.
+			for (const TPair<int32 /* NeighborVertex */, float /* NeighborArea */>& Item: NeighborVertexAreaLookup)
+			{
+				InfluenceIndexLookup.Reset();
+				VertexToInfluenceMap.MultiFind(Item.Key, InfluenceIndexLookup);
+
+				for (int32 InfluenceIndex: InfluenceIndexLookup)
+				{
+					const SkeletalMeshImportData::FRawBoneInfluence& Influence = InOutInfluences[InfluenceIndex];
+					VertexInfluences.FindOrAdd(Influence.BoneIndex, 0.0f) += Influence.Weight * Item.Value;
+				}
+			}
+
+			if (VertexInfluences.IsEmpty())
+			{
+				continue;
+			}
+
+			// Normalize the weights and add them to the list of alternate weights.
+			float SumWeight = Algo::Accumulate(VertexInfluences, 0.0f, [](float V, const TPair<int32, float>& Item) { return V + Item.Value; });
+
+			for (const TPair<int32, float>& Item: VertexInfluences)
+			{
+				VertexToInfluenceMap.Add(UnassignedVertex, InOutInfluences.Num());
+				InOutInfluences.Add({Item.Value / SumWeight, UnassignedVertex, Item.Key});
+			}
+
+			UnassignedVertexQueue.Remove(UnassignedVertex);
+			bPointRemoved = true;
+			break;
+		}
+
+		if (!bPointRemoved)
+		{
+			UE_LOG(LogLODUtilities, Warning, TEXT("Alternate skinning import: Some alternate skinning vertices could not be mapped onto the base mesh. These will be skinned to the root bone."));
+			break;
+		}
+	}
+}
+
 bool FLODUtilities::UpdateAlternateSkinWeights(
 	FSkeletalMeshLODModel& LODModelDest,
 	FSkeletalMeshImportData& ImportDataDest,
@@ -2207,29 +2358,30 @@ bool FLODUtilities::UpdateAlternateSkinWeights(
 	{
 		new(VertIndexAndZ)FIndexAndZ(VertexIndex, ImportDataDest.Points[VertexIndex]);
 	}
-	// Sort the vertices by z value
+	// Project the vertices onto a diagonal, use that to quickly find vertices for
+	// closest point matching. 
 	VertIndexAndZ.Sort(FCompareIndexAndZ());
-	
+
 	auto FindSimilarPosition = [&VertIndexAndZ, &ImportDataDest](const FVector3f& Position, TArray<int32>& PositionMatches, const float ComparisonThreshold)
 	{
 		PositionMatches.Reset();
-		FIndexAndZ PositionZ = FIndexAndZ(0, Position);
+		const FIndexAndZ PositionZ = FIndexAndZ(0, Position);
+		
+		// Use binary search to narrow down the range that matches our search area.
+		FIndexAndZ ZBound;
+		ZBound.Z = PositionZ.Z - ComparisonThreshold;
+		const int32 StartRange = Algo::LowerBound(VertIndexAndZ, ZBound, FCompareIndexAndZ{});
+
+		ZBound.Z = Position.Z + ComparisonThreshold;
+		const int32 EndRange = Algo::UpperBound(VertIndexAndZ, ZBound, FCompareIndexAndZ{});
+		
 		// Search for duplicates, quickly!
-		for (int32 i = 0; i < VertIndexAndZ.Num(); i++)
+		for (int32 Index = StartRange; Index < EndRange; Index++)
 		{
-			if (PositionZ.Z - ComparisonThreshold > VertIndexAndZ[i].Z)
-			{
-				continue;
-			}
-			else if (PositionZ.Z + ComparisonThreshold < VertIndexAndZ[i].Z)
-			{
-				break;
-			}
-
-			const FVector3f& PositionA = ImportDataDest.Points[VertIndexAndZ[i].Index];
+			const FVector3f& PositionA = ImportDataDest.Points[VertIndexAndZ[Index].Index];
 			if (PointsEqual(PositionA, Position, ComparisonThreshold))
 			{
-				PositionMatches.Add(VertIndexAndZ[i].Index);
+				PositionMatches.Add(VertIndexAndZ[Index].Index);
 			}
 		}
 	};
@@ -2330,7 +2482,7 @@ bool FLODUtilities::UpdateAlternateSkinWeights(
 				for (int32 MatchDestinationIndex = 0; MatchDestinationIndex < SimilarDestinationVertex.Num(); ++MatchDestinationIndex)
 				{
 					VertexMatchDest.VertexIndexes.Add(SimilarDestinationVertex[MatchDestinationIndex]);
-					VertexMatchDest.Ratios.Add(1.0f);
+					VertexMatchDest.Ratios.Add(1.0f / SimilarDestinationVertex.Num());
 				}
 			}
 		}
@@ -2352,7 +2504,6 @@ bool FLODUtilities::UpdateAlternateSkinWeights(
 					UE_LOG(LogLODUtilities, Warning, TEXT("Alternate skinning import: Cannot find a destination vertex index match for source vertex index. Alternate skinning quality will be lower."));
 					bNoMatchMsgDone = true;
 				}
-				continue;
 			}
 		}
 		bAllSourceVertexAreMatch = VertexIndexSrcToVertexIndexDestMatches.Num() == PointNumberSrc;
@@ -2385,7 +2536,6 @@ bool FLODUtilities::UpdateAlternateSkinWeights(
 		const SkeletalMeshImportData::FRawBoneInfluence& InfluenceSrc = ImportDataSrc.Influences[InfluenceIndexSrc];
 		int32 VertexIndexSource = InfluenceSrc.VertexIndex;
 		int32 BoneIndexSource = InfluenceSrc.BoneIndex;
-		float Weight = InfluenceSrc.Weight;
 		//We need to remap the source bone index to have the matching target bone index
 		int32 BoneIndexDest = RemapBoneIndexSrcToDest[BoneIndexSource];
 		if (BoneIndexDest != INDEX_NONE)
@@ -2403,21 +2553,24 @@ bool FLODUtilities::UpdateAlternateSkinWeights(
 			for (int32 ImpactedIndex = 0; ImpactedIndex < SourceVertexMatch->VertexIndexes.Num(); ++ImpactedIndex)
 			{
 				uint32 VertexIndexDest = SourceVertexMatch->VertexIndexes[ImpactedIndex];
-				float Ratio = SourceVertexMatch->Ratios[ImpactedIndex];
-				if (FMath::IsNearlyZero(Ratio, KINDA_SMALL_NUMBER))
+				const float Ratio = SourceVertexMatch->Ratios[ImpactedIndex];
+
+				const float InfluenceWeight = InfluenceSrc.Weight* Ratio;
+				if (InfluenceWeight < UE::AnimationCore::BoneWeightThreshold)
 				{
 					continue;
 				}
+				
 				SkeletalMeshImportData::FRawBoneInfluence AlternateInfluence;
 				AlternateInfluence.BoneIndex = BoneIndexDest;
 				AlternateInfluence.VertexIndex = VertexIndexDest;
-				AlternateInfluence.Weight = InfluenceSrc.Weight* Ratio;
+				AlternateInfluence.Weight = InfluenceWeight;
 				int32 AlternateInfluencesIndex = AlternateInfluences.Add(AlternateInfluence);
 				AlternateInfluencesMap.Add(AlternateInfluencesIndex);
 			}
 		}
 	}
-	
+
 	//In case the source geometry was not matching the destination we have to add influence for each extra destination vertex index
 	if (VertexIndexDestToVertexIndexSrcMatches.Num() > 0)
 	{
@@ -2493,9 +2646,30 @@ bool FLODUtilities::UpdateAlternateSkinWeights(
 		}
 	}
 
+	// Check if there are any points that have no weight assignments after all is said and done. Attempt to average across neighbors to reconstruct
+	// those weights. Currently this relies on these points being completely surrounded by at least two points that have weights assigned to them.
+	TBitArray<> UnassignedVertexBitmask(true, ImportDataDest.Points.Num());
+	for (SkeletalMeshImportData::FRawBoneInfluence& Influence: AlternateInfluences)
+	{
+		UnassignedVertexBitmask[Influence.VertexIndex] = false;
+	}
+
+	TArray<int32> UnassignedVertexes;
+	for (TConstSetBitIterator<> It(UnassignedVertexBitmask); It; ++It)
+	{
+		UnassignedVertexes.Add(It.GetIndex());
+	}
+
+	if (!UnassignedVertexes.IsEmpty())
+	{
+		InfillMissingVertexSkinWeights(ImportDataDest, UnassignedVertexes, AlternateInfluences);
+	}
+
 	//Sort and normalize weights for alternate influences
 	ProcessImportMeshInfluences(ImportDataDest.Wedges.Num(), AlternateInfluences, SkeletalMeshDest->GetPathName());
 
+	
+
 	//Store the remapped influence into the profile, the function SkeletalMeshTools::ChunkSkinnedVertices will use all profiles including this one to chunk the sections
 	FImportedSkinWeightProfileData& ImportedProfileData = LODModelDest.SkinWeightProfiles.Add(ProfileNameDest);
 	ImportedProfileData.SourceModelInfluences.Empty(AlternateInfluences.Num());
