Interchange: Material workflow implementation

#rnx

#virtualized

diff --git a/Engine/Plugins/Interchange/Runtime/Source/FactoryNodes/Private/InterchangeMaterialFactoryNode.cpp b/Engine/Plugins/Interchange/Runtime/Source/FactoryNodes/Private/InterchangeMaterialFactoryNode.cpp
--- a/Engine/Plugins/Interchange/Runtime/Source/FactoryNodes/Private/InterchangeMaterialFactoryNode.cpp
+++ b/Engine/Plugins/Interchange/Runtime/Source/FactoryNodes/Private/InterchangeMaterialFactoryNode.cpp
@@ -20,6 +20,16 @@ FString UInterchangeBaseMaterialFactoryNode::GetMaterialFactoryNodeUidFromMateri
 	return NewUid;
 }
 
+bool UInterchangeBaseMaterialFactoryNode::GetCustomIsMaterialImportEnabled(bool& AttributeValue) const
+{
+	IMPLEMENT_NODE_ATTRIBUTE_GETTER(IsMaterialImportEnabled, bool);
+}
+
+bool UInterchangeBaseMaterialFactoryNode::SetCustomIsMaterialImportEnabled(const bool& AttributeValue)
+{
+	IMPLEMENT_NODE_ATTRIBUTE_SETTER_NODELEGATE(IsMaterialImportEnabled, bool);
+}
+
 FString UInterchangeMaterialFactoryNode::GetTypeName() const
 {
 	const FString TypeName = TEXT("MaterialFactoryNode");
diff --git a/Engine/Plugins/Interchange/Runtime/Source/FactoryNodes/Public/InterchangeMaterialFactoryNode.h b/Engine/Plugins/Interchange/Runtime/Source/FactoryNodes/Public/InterchangeMaterialFactoryNode.h
--- a/Engine/Plugins/Interchange/Runtime/Source/FactoryNodes/Public/InterchangeMaterialFactoryNode.h
+++ b/Engine/Plugins/Interchange/Runtime/Source/FactoryNodes/Public/InterchangeMaterialFactoryNode.h
@@ -20,6 +20,16 @@ class INTERCHANGEFACTORYNODES_API UInterchangeBaseMaterialFactoryNode : public U
 
 public:
 	static FString GetMaterialFactoryNodeUidFromMaterialNodeUid(const FString& TranslatedNodeUid);
+
+	UFUNCTION(BlueprintCallable, Category = "Interchange | Node | Material")
+	bool GetCustomIsMaterialImportEnabled(bool& AttributeValue) const;
+
+	/** Sets if this shader graph should be rendered two sided or not. Defaults to off. */
+	UFUNCTION(BlueprintCallable, Category = "Interchange | Node | Material")
+	bool SetCustomIsMaterialImportEnabled(const bool& AttributeValue);
+
+private:
+	const UE::Interchange::FAttributeKey Macro_CustomIsMaterialImportEnabledKey = UE::Interchange::FAttributeKey(TEXT("IsMaterialImportEnabled"));
 };
 
 UCLASS(BlueprintType)
diff --git a/Engine/Plugins/Interchange/Runtime/Source/Import/Private/Material/InterchangeMaterialFactory.cpp b/Engine/Plugins/Interchange/Runtime/Source/Import/Private/Material/InterchangeMaterialFactory.cpp
--- a/Engine/Plugins/Interchange/Runtime/Source/Import/Private/Material/InterchangeMaterialFactory.cpp
+++ b/Engine/Plugins/Interchange/Runtime/Source/Import/Private/Material/InterchangeMaterialFactory.cpp
@@ -824,8 +824,10 @@ UInterchangeFactoryBase::FImportAssetResult UInterchangeMaterialFactory::BeginIm
 		}
 	}
 
+	bool bCanImportMaterial = true;
+	MaterialFactoryNode->GetCustomIsMaterialImportEnabled(bCanImportMaterial);
 	// create a new material or overwrite existing asset, if possible
-	if (!ExistingAsset)
+	if (!ExistingAsset && bCanImportMaterial)
 	{
 		if (MaterialClass->IsChildOf<UMaterialInstanceDynamic>())
 		{
@@ -859,7 +861,10 @@ UInterchangeFactoryBase::FImportAssetResult UInterchangeMaterialFactory::BeginIm
 
 	if (!Material)
 	{
-		CouldNotCreateMaterialLog(LOCTEXT("MatFactory_CouldNotCreateMat_MaterialCreationFail", "Material creation fail."));
+		if (bCanImportMaterial)
+		{
+			CouldNotCreateMaterialLog(LOCTEXT("MatFactory_CouldNotCreateMat_MaterialCreationFail", "Material creation fail."));
+		}
 		return ImportAssetResult;
 	}
 
diff --git a/Engine/Plugins/Interchange/Runtime/Source/Import/Private/Mesh/InterchangeSkeletalMeshFactory.cpp b/Engine/Plugins/Interchange/Runtime/Source/Import/Private/Mesh/InterchangeSkeletalMeshFactory.cpp
--- a/Engine/Plugins/Interchange/Runtime/Source/Import/Private/Mesh/InterchangeSkeletalMeshFactory.cpp
+++ b/Engine/Plugins/Interchange/Runtime/Source/Import/Private/Mesh/InterchangeSkeletalMeshFactory.cpp
@@ -1214,8 +1214,9 @@ UInterchangeFactoryBase::FImportAssetResult UInterchangeSkeletalMeshFactory::Imp
 		FSkeletalMaterial* SkeletalMaterial = Materials.FindByPredicate([&MaterialSlotName](const FSkeletalMaterial& Material) { return Material.MaterialSlotName == MaterialSlotName; });
 		if (SkeletalMaterial)
 		{
-			//When we do a reimport we update the material interface only if the specified MaterialInterface is not null
-			if (!bIsReImport || MaterialInterface || !SkeletalMaterial->MaterialInterface)
+			//When we are not re-importing, we always force update the material, we should see this case when importing LODs is on since its an import.
+			//When we do a re-import we update the material interface only if the current asset matching material is null and is not the default material.
+			if (!bIsReImport || (MaterialInterface && (!SkeletalMaterial->MaterialInterface || SkeletalMaterial->MaterialInterface == UMaterial::GetDefaultMaterial(MD_Surface))))
 			{
 				SkeletalMaterial->MaterialInterface = NewMaterial;
 			}
@@ -1237,7 +1238,7 @@ UInterchangeFactoryBase::FImportAssetResult UInterchangeSkeletalMeshFactory::Imp
 		FName MaterialSlotName = *SlotMaterialDependency.Key;
 
 		const UInterchangeBaseMaterialFactoryNode* MaterialFactoryNode = Cast<UInterchangeBaseMaterialFactoryNode>(Arguments.NodeContainer->GetNode(SlotMaterialDependency.Value));
-		if (!MaterialFactoryNode || !MaterialFactoryNode->IsEnabled())
+		if (!MaterialFactoryNode)
 		{
 			UpdateOrAddSkeletalMaterial(MaterialSlotName, nullptr);
 			continue;
@@ -1633,6 +1634,17 @@ UInterchangeFactoryBase::FImportAssetResult UInterchangeSkeletalMeshFactory::End
 					Parameters.ImportedMaterials = &ImportAssetObjectLODData.ImportedMaterials;
 					Parameters.ExistingOriginalPerSectionMaterialImportName = &ImportAssetObjectLODData.ExistingOriginalPerSectionMaterialImportName;
 					FLODUtilities::MatchImportedMaterials(Parameters);
+					//Flush the old LOD 0 sections after we rematch the materials
+					if (Parameters.LodIndex == 0 && Parameters.bIsReImport)
+					{
+						if (FSkeletalMeshModel* ImportedModel = SkeletalMesh->GetImportedModel())
+						{
+							if (ImportedModel->LODModels.IsValidIndex(0))
+							{
+								SkeletalMesh->GetImportedModel()->LODModels[0].Sections.Empty();
+							}
+						}
+					}
 				}
 			}
 		}
diff --git a/Engine/Plugins/Interchange/Runtime/Source/Import/Private/Mesh/InterchangeStaticMeshFactory.cpp b/Engine/Plugins/Interchange/Runtime/Source/Import/Private/Mesh/InterchangeStaticMeshFactory.cpp
--- a/Engine/Plugins/Interchange/Runtime/Source/Import/Private/Mesh/InterchangeStaticMeshFactory.cpp
+++ b/Engine/Plugins/Interchange/Runtime/Source/Import/Private/Mesh/InterchangeStaticMeshFactory.cpp
@@ -45,13 +45,6 @@
 #include "EditorFramework/AssetImportData.h"
 #endif //WITH_EDITORONLY_DATA
 
-static bool GInterchangeStaticMeshReorderMaterialSlots = true;
-static FAutoConsoleVariableRef CCvarInterchangeStaticMeshReorderMaterialSlots(
-	TEXT("Interchange.FeatureFlags.Import.StaticMesh.ReorderMaterialSlots"),
-	GInterchangeStaticMeshReorderMaterialSlots,
-	TEXT("Whether Re-importing a static mesh should reorder the material slots."),
-	ECVF_Default);
-
 UClass* UInterchangeStaticMeshFactory::GetFactoryClass() const
 {
 	return UStaticMesh::StaticClass();
@@ -80,7 +73,7 @@ namespace UE::Interchange::Private::StaticMesh
 		}
 		
 		FStaticMeshConstAttributes StaticMeshAttributes(BaseMeshDescription);
-		TPolygonGroupAttributesRef<const FName> SlotNames = StaticMeshAttributes.GetPolygonGroupMaterialSlotNames();
+		TPolygonGroupAttributesConstRef<FName> SlotNames = StaticMeshAttributes.GetPolygonGroupMaterialSlotNames();
 
 		TArray<int32> RemapMaterialIndexes;
 		RemapMaterialIndexes.Reserve(Materials.Num());
@@ -276,8 +269,9 @@ UInterchangeFactoryBase::FImportAssetResult UInterchangeStaticMeshFactory::Impor
 		FStaticMaterial* StaticMaterial = StaticMesh->GetStaticMaterials().FindByPredicate([&MaterialSlotName](const FStaticMaterial& Material) { return Material.MaterialSlotName == MaterialSlotName; });
 		if (StaticMaterial)
 		{
-			//When we do a re-import we update the material interface only if the specified MaterialInterface is not null
-			if (!bReimport || MaterialInterface || !StaticMaterial->MaterialInterface)
+			//When we are not re-importing, we always force update the material, we should see this case when importing LODs is on since its an import.
+			//When we do a re-import we update the material interface only if the current asset matching material is null and is not the default material. (this avoid touching a slot that was change by the user)
+			if (!bReimport || (MaterialInterface && (!StaticMaterial->MaterialInterface || StaticMaterial->MaterialInterface == UMaterial::GetDefaultMaterial(MD_Surface))))
 			{
 				StaticMaterial->MaterialInterface = NewMaterial;
 			}
@@ -299,7 +293,7 @@ UInterchangeFactoryBase::FImportAssetResult UInterchangeStaticMeshFactory::Impor
 		FName MaterialSlotName = *SlotMaterialDependency.Key;
 
 		const UInterchangeBaseMaterialFactoryNode* MaterialFactoryNode = Cast<UInterchangeBaseMaterialFactoryNode>(Arguments.NodeContainer->GetNode(SlotMaterialDependency.Value));
-		if (!MaterialFactoryNode || !MaterialFactoryNode->IsEnabled())
+		if (!MaterialFactoryNode)
 		{
 			UpdateOrAddStaticMaterial(MaterialSlotName, nullptr);
 			continue;
@@ -531,6 +525,7 @@ UInterchangeFactoryBase::FImportAssetResult UInterchangeStaticMeshFactory::EndIm
 			//Match the existing section info map data
 
 			//First find the old mesh description polygon groups name that match with the imported mesh description polygon groups name.
+			//Copy the data
 			const int32 PreviousSectionCount = StaticMesh->GetSectionInfoMap().GetSectionNumber(LodIndex);
 			TMap<FPolygonGroupID, FPolygonGroupID> ImportedToOldPolygonGroupMatch;
 			ImportedToOldPolygonGroupMatch.Reserve(LodMeshDescription.PolygonGroups().Num());
@@ -573,10 +568,15 @@ UInterchangeFactoryBase::FImportAssetResult UInterchangeStaticMeshFactory::EndIm
 					NewSectionInfoMapData.Add(FMeshSectionInfo(MaterialSlotIndex));
 				}
 			}
-			//Recreate the section info map
+
+			//Clear all section for this LOD
+			for (int32 PreviousSectionIndex = 0; PreviousSectionIndex < PreviousSectionCount; ++PreviousSectionIndex)
+			{
+				StaticMesh->GetSectionInfoMap().Remove(LodIndex, PreviousSectionIndex);
+			}
+			//Recreate the new section info map
 			for (int32 NewSectionIndex = 0; NewSectionIndex < NewSectionInfoMapData.Num(); ++NewSectionIndex)
 			{
-				StaticMesh->GetSectionInfoMap().Remove(LodIndex, NewSectionIndex);
 				StaticMesh->GetSectionInfoMap().Set(LodIndex, NewSectionIndex, NewSectionInfoMapData[NewSectionIndex]);
 			}
 		}
@@ -678,9 +678,12 @@ UInterchangeFactoryBase::FImportAssetResult UInterchangeStaticMeshFactory::EndIm
 	FMeshBudgetProjectSettingsUtils::SetLodGroupForStaticMesh(StaticMesh);
 #endif
 
-	if (bReimport && GInterchangeStaticMeshReorderMaterialSlots)
+	if (bReimport)
 	{
 		UE::Interchange::Private::StaticMesh::ReorderMaterialSlotToBaseLod(StaticMesh);
+#if WITH_EDITOR
+		UStaticMesh::RemoveUnusedMaterialSlots(StaticMesh);
+#endif
 	}
 
 	ImportAssetResult.ImportedObject = StaticMesh;
diff --git a/Engine/Plugins/Interchange/Runtime/Source/Pipelines/Private/InterchangeGenericMaterialPipeline.cpp b/Engine/Plugins/Interchange/Runtime/Source/Pipelines/Private/InterchangeGenericMaterialPipeline.cpp
--- a/Engine/Plugins/Interchange/Runtime/Source/Pipelines/Private/InterchangeGenericMaterialPipeline.cpp
+++ b/Engine/Plugins/Interchange/Runtime/Source/Pipelines/Private/InterchangeGenericMaterialPipeline.cpp
@@ -2,8 +2,9 @@
 
 #include "InterchangeGenericMaterialPipeline.h"
 
+#include "AssetRegistry/AssetData.h"
+#include "AssetRegistry/AssetRegistryModule.h"
 #include "CoreMinimal.h"
-
 #include "InterchangeGenericTexturePipeline.h"
 #include "InterchangeMaterialDefinitions.h"
 #include "InterchangeMaterialFactoryNode.h"
@@ -21,7 +22,6 @@
 #include "Nodes/InterchangeBaseNode.h"
 #include "Nodes/InterchangeUserDefinedAttribute.h"
 #include "InterchangeMaterialInstanceNode.h"
-
 #include "Materials/MaterialExpressionAdd.h"
 #include "Materials/MaterialExpressionComponentMask.h"
 #include "Materials/MaterialExpressionDivide.h"
@@ -202,6 +202,125 @@ namespace UE::Interchange::InterchangeGenericMaterialPipeline::Private
 		
 		return Expression;
 	}
+
+	UMaterialInterface* FindExistingMaterial(const FString& BasePath, const FString& MaterialFullName, const bool bRecursivePaths)
+	{
+		UMaterialInterface* Material = nullptr;
+		FAssetRegistryModule& AssetRegistryModule = FModuleManager::LoadModuleChecked<FAssetRegistryModule>("AssetRegistry");
+		IAssetRegistry& AssetRegistry = AssetRegistryModule.Get();
+
+		// Finish/update any scans
+		TArray<FString> ScanPaths;
+		if (BasePath.IsEmpty() || BasePath == TEXT("/"))
+		{
+			FPackageName::QueryRootContentPaths(ScanPaths);
+		}
+		else
+		{
+			ScanPaths.Add(BasePath);
+		}
+		const bool bForceRescan = false;
+		AssetRegistry.ScanPathsSynchronous(ScanPaths, bForceRescan);
+
+
+		FARFilter Filter;
+		Filter.bRecursiveClasses = true;
+		Filter.bRecursivePaths = bRecursivePaths;
+		Filter.ClassPaths.Add(UMaterialInterface::StaticClass()->GetClassPathName());
+		Filter.PackagePaths.Add(FName(*BasePath));
+
+		TArray<FAssetData> AssetData;
+		AssetRegistry.GetAssets(Filter, AssetData);
+
+		TArray<UMaterialInterface*> FoundMaterials;
+		for (const FAssetData& Data : AssetData)
+		{
+			if (Data.AssetName == FName(*MaterialFullName))
+			{
+				Material = Cast<UMaterialInterface>(Data.GetAsset());
+				if (Material != nullptr)
+				{
+					FoundMaterials.Add(Material);
+				}
+			}
+		}
+
+		return FoundMaterials.Num() > 0 ? FoundMaterials[0] : Material;
+	}
+
+	UMaterialInterface* FindExistingMaterialFromSearchLocation(const FString& MaterialFullName, const FString& ContentPath, EInterchangeMaterialSearchLocation SearchLocation)
+	{
+		if (SearchLocation == EInterchangeMaterialSearchLocation::DoNotSearch)
+		{
+			return nullptr;
+		}
+
+		//Search in memory
+		constexpr bool bExactClass = false;
+		UMaterialInterface* FoundMaterial = nullptr;
+		//We search only in memory for search in local folder.
+		if (SearchLocation == EInterchangeMaterialSearchLocation::Local)
+		{
+			FoundMaterial = FindObject<UMaterialInterface>(nullptr, *MaterialFullName, bExactClass);
+			if(FoundMaterial)
+			{
+				//Make sure the path of the material in memory is local
+				FString PackagePath = FoundMaterial->GetPackage()->GetPathName();
+				if (!PackagePath.Equals(ContentPath))
+				{
+					FoundMaterial = nullptr;
+				}
+			}
+		}
+
+		if (FoundMaterial == nullptr)
+		{
+			FString SearchPath = ContentPath;
+
+			// Search in asset's local folder
+			FoundMaterial = FindExistingMaterial(SearchPath, MaterialFullName, false);
+
+			// Search recursively in asset's folder
+			if (FoundMaterial == nullptr &&
+				(SearchLocation != EInterchangeMaterialSearchLocation::Local))
+			{
+				FoundMaterial = FindExistingMaterial(SearchPath, MaterialFullName, true);
+			}
+
+			if (FoundMaterial == nullptr &&
+				(SearchLocation == EInterchangeMaterialSearchLocation::UnderParent ||
+					SearchLocation == EInterchangeMaterialSearchLocation::UnderRoot ||
+					SearchLocation == EInterchangeMaterialSearchLocation::AllAssets))
+			{
+				// Search recursively in parent's folder
+				SearchPath = FPaths::GetPath(SearchPath);
+				if (!SearchPath.IsEmpty())
+				{
+					FoundMaterial = FindExistingMaterial(SearchPath, MaterialFullName, true);
+				}
+			}
+			if (FoundMaterial == nullptr &&
+				(SearchLocation == EInterchangeMaterialSearchLocation::UnderRoot ||
+					SearchLocation == EInterchangeMaterialSearchLocation::AllAssets))
+			{
+				// Search recursively in root folder of asset
+				FString OutPackageRoot, OutPackagePath, OutPackageName;
+				FPackageName::SplitLongPackageName(SearchPath, OutPackageRoot, OutPackagePath, OutPackageName);
+				if (!SearchPath.IsEmpty())
+				{
+					FoundMaterial = FindExistingMaterial(OutPackageRoot, MaterialFullName, true);
+				}
+			}
+			if (FoundMaterial == nullptr &&
+				SearchLocation == EInterchangeMaterialSearchLocation::AllAssets)
+			{
+				// Search everywhere
+				FoundMaterial = FindExistingMaterial(TEXT("/"), MaterialFullName, true);
+			}
+		}
+
+		return FoundMaterial;
+	}
 }
 
 namespace UE::Interchange::Materials::HashUtils
@@ -494,8 +613,7 @@ void UInterchangeGenericMaterialPipeline::AdjustSettingsForContext(EInterchangeP
 #if WITH_EDITOR
 	TArray<FString> HideCategories;
 	bool bIsObjectAMaterial = !ReimportAsset ? false : ReimportAsset->IsA(UMaterialInterface::StaticClass());
-	if ((!bIsObjectAMaterial && ImportType == EInterchangePipelineContext::AssetReimport)
-		|| ImportType == EInterchangePipelineContext::AssetCustomLODImport
+	if (ImportType == EInterchangePipelineContext::AssetCustomLODImport
 		|| ImportType == EInterchangePipelineContext::AssetCustomLODReimport
 		|| ImportType == EInterchangePipelineContext::AssetAlternateSkinningImport
 		|| ImportType == EInterchangePipelineContext::AssetAlternateSkinningReimport)
@@ -510,7 +628,18 @@ void UInterchangeGenericMaterialPipeline::AdjustSettingsForContext(EInterchangeP
 		{
 			HidePropertiesOfCategory(OuterMostPipeline, this, HideCategoryName);
 		}
+		if (!bIsObjectAMaterial && ImportType == EInterchangePipelineContext::AssetReimport)
+		{
+			//When we re-import we hide all setting but search location, so we can find existing materials.
+			HideProperty(OuterMostPipeline, this, GET_MEMBER_NAME_CHECKED(UInterchangeGenericMaterialPipeline, bImportMaterials));
+			HideProperty(OuterMostPipeline, this, GET_MEMBER_NAME_CHECKED(UInterchangeGenericMaterialPipeline, MaterialImport));
+			HideProperty(OuterMostPipeline, this, GET_MEMBER_NAME_CHECKED(UInterchangeGenericMaterialPipeline, bIdentifyDuplicateMaterials));
+			HideProperty(OuterMostPipeline, this, GET_MEMBER_NAME_CHECKED(UInterchangeGenericMaterialPipeline, bCreateMaterialInstanceForParent));
+			HideProperty(OuterMostPipeline, this, GET_MEMBER_NAME_CHECKED(UInterchangeGenericMaterialPipeline, ParentMaterial));
+			HideProperty(OuterMostPipeline, this, GET_MEMBER_NAME_CHECKED(UInterchangeGenericMaterialPipeline, AssetName));
+		}
 	}
+	
 #endif //WITH_EDITOR
 	using namespace UE::Interchange;
 
@@ -599,15 +728,9 @@ void UInterchangeGenericMaterialPipeline::ExecutePipeline(UInterchangeBaseNodeCo
 		TexturePipeline->ScriptedExecutePipeline(InBaseNodeContainer, InSourceDatas, ContentBasePath);
 	}
 
-	//Skip Material import if the toggle is off
-	if (!bImportMaterials)
-	{
-		return;
-	}
 	
 	MaterialNodes.Empty();
 	MaterialFactoryNodes.Empty();
-
 	//Find all translated node we need for this pipeline
 	BaseNodeContainer->IterateNodes([this](const FString& NodeUid, UInterchangeBaseNode* Node)
 	{
@@ -669,11 +792,6 @@ void UInterchangeGenericMaterialPipeline::ExecutePipeline(UInterchangeBaseNodeCo
 				/* Clearing the AttributeStorageNode as it might affect how the MaterialFunctionsFactories are created. */
 				AttributeStorageNode = nullptr;
 			}
-
-			if (MaterialBaseFactoryNode)
-			{
-				MaterialBaseFactoryNode->SetEnabled(bImportUnusedMaterial);
-			}
 		}
 	}
 	else if (MaterialImport == EInterchangeMaterialImportOption::ImportAsMaterialInstances)
@@ -687,6 +805,35 @@ void UInterchangeGenericMaterialPipeline::ExecutePipeline(UInterchangeBaseNodeCo
 		}
 	}
 
+	
+	BaseNodeContainer->IterateNodesOfType<UInterchangeBaseMaterialFactoryNode>([&ContentBasePath, bClosureImportMaterials = bImportMaterials || bImportUnusedMaterial, ClosureSearchLocation = SearchLocation](const FString& NodeUid, UInterchangeBaseMaterialFactoryNode* MaterialBaseFactoryNode)
+		{
+			if (MaterialBaseFactoryNode)
+			{
+				MaterialBaseFactoryNode->SetCustomIsMaterialImportEnabled(bClosureImportMaterials);
+				//Disable all materials if the toggle is off
+				if (!bClosureImportMaterials)
+				{
+					MaterialBaseFactoryNode->SetEnabled(false);
+				}
+				
+				//See if we can assign an existing material from the search location
+				FString MaterialName = MaterialBaseFactoryNode->GetDisplayLabel();
+				if (UMaterialInterface* ExistingMaterial = UE::Interchange::InterchangeGenericMaterialPipeline::Private::FindExistingMaterialFromSearchLocation(MaterialName, ContentBasePath, ClosureSearchLocation))
+				{
+					//Make sure we have the correct type of material (can be material instance) before setting the custom object reference.
+					if ((MaterialBaseFactoryNode->IsA<UInterchangeMaterialInstanceFactoryNode>() && ExistingMaterial->IsA<UMaterialInstance>())
+						|| (MaterialBaseFactoryNode->IsA<UInterchangeMaterialFactoryNode>() && ExistingMaterial->IsA<UMaterial>()))
+					{
+						MaterialBaseFactoryNode->SetCustomReferenceObject(ExistingMaterial);
+						//No need to import an existing material
+						MaterialBaseFactoryNode->SetCustomIsMaterialImportEnabled(false);
+						MaterialBaseFactoryNode->SetEnabled(false);
+					}
+				}
+			}
+		});
+
 	TArray<UInterchangeMaterialInstanceNode*> MaterialInstanceNodes;
 	BaseNodeContainer->IterateNodesOfType<UInterchangeMaterialInstanceNode>([&MaterialInstanceNodes](const FString& NodeUid, UInterchangeMaterialInstanceNode* MaterialNode)
 		{
diff --git a/Engine/Plugins/Interchange/Runtime/Source/Pipelines/Public/InterchangeGenericMaterialPipeline.h b/Engine/Plugins/Interchange/Runtime/Source/Pipelines/Public/InterchangeGenericMaterialPipeline.h
--- a/Engine/Plugins/Interchange/Runtime/Source/Pipelines/Public/InterchangeGenericMaterialPipeline.h
+++ b/Engine/Plugins/Interchange/Runtime/Source/Pipelines/Public/InterchangeGenericMaterialPipeline.h
@@ -36,6 +36,21 @@ enum class EInterchangeMaterialImportOption : uint8
 	ImportAsMaterialInstances,
 };
 
+UENUM(BlueprintType)
+enum class EInterchangeMaterialSearchLocation : uint8
+{
+	/** Search for existing material in local import folder only. */
+	Local,
+	/** Search for existing material recursively from parent folder. */
+	UnderParent,
+	/** Search for existing material recursively from root folder. */
+	UnderRoot,
+	/** Search for existing material in all assets folders. */
+	AllAssets,
+	/** Do not search for existing existing materials */
+	DoNotSearch,
+};
+
 UCLASS(BlueprintType, editinlinenew)
 class INTERCHANGEPIPELINES_API UInterchangeGenericMaterialPipeline : public UInterchangePipelineBase
 {
@@ -48,6 +63,10 @@ public:
 	UPROPERTY(EditAnywhere, BlueprintReadWrite, Category = "Materials")
 	bool bImportMaterials = true;
 
+	/** Specify where we should search for existing materials when importing.*/
+	UPROPERTY(EditAnywhere, BlueprintReadWrite, Category = "Materials")
+	EInterchangeMaterialSearchLocation SearchLocation = EInterchangeMaterialSearchLocation::Local;
+
 	/** If not empty, and there is only one asset and one source data, we will name the asset with this string. */
 	UPROPERTY(EditAnywhere, BlueprintReadWrite, Category = "Materials", meta = (StandAlonePipelineProperty = "True", AlwaysResetToDefault = "True"))
 	FString AssetName;
@@ -56,7 +75,7 @@ public:
 	EInterchangeMaterialImportOption MaterialImport = EInterchangeMaterialImportOption::ImportAsMaterials;
 	
 	/** If set, reference materials along with respective material instances would be created*/
-	UPROPERTY(EditAnywhere, BlueprintReadWrite, Category = "Materials")
+	UPROPERTY(EditAnywhere, BlueprintReadWrite, Category = "Materials", Meta = (EditCondition = "bImportMaterials"))
 	bool bIdentifyDuplicateMaterials = false;
 
 	/** If set, additional material instance would be created for reference/parent material*/
diff --git a/Engine/Plugins/Interchange/Runtime/Source/Pipelines/Public/InterchangePipelineMeshesUtilities.h b/Engine/Plugins/Interchange/Runtime/Source/Pipelines/Public/InterchangePipelineMeshesUtilities.h
--- a/Engine/Plugins/Interchange/Runtime/Source/Pipelines/Public/InterchangePipelineMeshesUtilities.h
+++ b/Engine/Plugins/Interchange/Runtime/Source/Pipelines/Public/InterchangePipelineMeshesUtilities.h
@@ -345,7 +345,9 @@ namespace UE::Interchange::MeshesUtilities
 			FactoryNode.SetSlotMaterialDependencyUid(SlotMaterialDependency.Key, MaterialFactoryNodeUid);
 			if (UInterchangeBaseMaterialFactoryNode* MaterialFactoryNode = Cast<UInterchangeBaseMaterialFactoryNode>(NodeContainer.GetFactoryNode(MaterialFactoryNodeUid)))
 			{
-				MaterialFactoryNode->SetEnabled(true);
+				bool IsMaterialImportEnabled = true;
+				MaterialFactoryNode->GetCustomIsMaterialImportEnabled(IsMaterialImportEnabled);
+				MaterialFactoryNode->SetEnabled(IsMaterialImportEnabled);
 
 				// Create a factory dependency so Material asset are imported before the static mesh asset
 				TArray<FString> FactoryDependencies;
diff --git a/Engine/Source/Developer/SkeletalMeshUtilitiesCommon/Private/LODUtilities.cpp b/Engine/Source/Developer/SkeletalMeshUtilitiesCommon/Private/LODUtilities.cpp
--- a/Engine/Source/Developer/SkeletalMeshUtilitiesCommon/Private/LODUtilities.cpp
+++ b/Engine/Source/Developer/SkeletalMeshUtilitiesCommon/Private/LODUtilities.cpp
@@ -52,13 +52,6 @@ IMPLEMENT_MODULE(FDefaultModuleImpl, SkeletalMeshUtilitiesCommon)
 
 DEFINE_LOG_CATEGORY_STATIC(LogLODUtilities, Log, All);
 
-static bool GInterchangeSkeletalMeshReorderMaterialSlots = true;
-static FAutoConsoleVariableRef CCvarInterchangeSkeletalMeshReorderMaterialSlots(
-	TEXT("Interchange.FeatureFlags.Import.SkeletalMesh.ReorderMaterialSlots"),
-	GInterchangeSkeletalMeshReorderMaterialSlots,
-	TEXT("Whether Re-importing a skeletalmesh should reorder the material slots."),
-	ECVF_Default);
-
 /**
 * Process and update the vertex Influences using the predefined wedges
 *
@@ -3850,6 +3843,7 @@ void FLODUtilities::MatchImportedMaterials(FLODUtilities::FSkeletalMeshMatchImpo
 	TArray<FSkeletalMaterial>& Materials = Parameters.SkeletalMesh->GetMaterials();
 
 	TMap<FName, int32> LODMaterialMapRedirection;
+	TMap<FName, FSkelMeshSourceSectionUserData> NewUserSectionsDataMap;
 	//If we reimport we have to keep the existing user section info data. We use the material name to match the existing
 	if (Parameters.bIsReImport)
 	{
@@ -3859,46 +3853,45 @@ void FLODUtilities::MatchImportedMaterials(FLODUtilities::FSkeletalMeshMatchImpo
 			, int32 SectionIndex
 			, int32& OutMaterialIndex
 			, const TArray<FName>& PerSectionMaterialImportName)->FName
-		{
-			const TArray<FSkeletalMaterial>& MeshMaterials = Parameters.SkeletalMesh->GetMaterials();
-			if (!ensure(MeshMaterials.Num() > 0))
-			{
-				OutMaterialIndex = -1;
-				return NAME_None;
-			}
-			if (PerSectionMaterialImportName.IsValidIndex(SectionIndex))
 			{
-				for (int32 MaterialIndex = 0; MaterialIndex < MeshMaterials.Num(); ++MaterialIndex)
+				const TArray<FSkeletalMaterial>& MeshMaterials = Parameters.SkeletalMesh->GetMaterials();
+				if (!ensure(MeshMaterials.Num() > 0))
+				{
+					OutMaterialIndex = -1;
+					return NAME_None;
+				}
+				if (PerSectionMaterialImportName.IsValidIndex(SectionIndex))
 				{
-					const FSkeletalMaterial& Material = MeshMaterials[MaterialIndex];
-					if (PerSectionMaterialImportName[SectionIndex] == Material.ImportedMaterialSlotName)
+					for (int32 MaterialIndex = 0; MaterialIndex < MeshMaterials.Num(); ++MaterialIndex)
 					{
-						OutMaterialIndex = MaterialIndex;
-						break;
+						const FSkeletalMaterial& Material = MeshMaterials[MaterialIndex];
+						if (PerSectionMaterialImportName[SectionIndex] == Material.ImportedMaterialSlotName)
+						{
+							OutMaterialIndex = MaterialIndex;
+							break;
+						}
 					}
+					return PerSectionMaterialImportName[SectionIndex];
 				}
-				return PerSectionMaterialImportName[SectionIndex];
-			}
 
-			OutMaterialIndex = Section.MaterialIndex;
-			if (LodInfo->LODMaterialMap.IsValidIndex(SectionIndex) && LodInfo->LODMaterialMap[SectionIndex] != INDEX_NONE)
-			{
-				OutMaterialIndex = LodInfo->LODMaterialMap[SectionIndex];
-			}
-			FName ImportedMaterialSlotName = NAME_None;
-			if (MeshMaterials.IsValidIndex(OutMaterialIndex))
-			{
-				ImportedMaterialSlotName = MeshMaterials[OutMaterialIndex].ImportedMaterialSlotName;
-			}
-			else
-			{
-				ImportedMaterialSlotName = MeshMaterials[0].ImportedMaterialSlotName;
-				OutMaterialIndex = 0;
-			}
-			return ImportedMaterialSlotName;
-		};
+				OutMaterialIndex = Section.MaterialIndex;
+				if (LodInfo->LODMaterialMap.IsValidIndex(SectionIndex) && LodInfo->LODMaterialMap[SectionIndex] != INDEX_NONE)
+				{
+					OutMaterialIndex = LodInfo->LODMaterialMap[SectionIndex];
+				}
+				FName ImportedMaterialSlotName = NAME_None;
+				if (MeshMaterials.IsValidIndex(OutMaterialIndex))
+				{
+					ImportedMaterialSlotName = MeshMaterials[OutMaterialIndex].ImportedMaterialSlotName;
+				}
+				else
+				{
+					ImportedMaterialSlotName = MeshMaterials[0].ImportedMaterialSlotName;
+					OutMaterialIndex = 0;
+				}
+				return ImportedMaterialSlotName;
+			};
 
-		TMap<FName, FSkelMeshSourceSectionUserData> NewUserSectionsDataMap;
 		FSkeletalMeshLODModel& ExistLODModel = Parameters.SkeletalMesh->GetImportedModel()->LODModels[Parameters.LodIndex];
 		for (int32 NewSectionIndex = 0; NewSectionIndex < ImportedMaterials.Num(); NewSectionIndex++)
 		{
@@ -3930,13 +3923,13 @@ void FLODUtilities::MatchImportedMaterials(FLODUtilities::FSkeletalMeshMatchImpo
 					, EmptyArray);
 
 				auto CopySectionData = [&ExistingSection, &SourceSectionUserData]()
-				{
-					//Set the user section data to reflect the existing settings
-					SourceSectionUserData.bCastShadow = ExistingSection.bCastShadow;
-					SourceSectionUserData.bVisibleInRayTracing = ExistingSection.bVisibleInRayTracing;
-					SourceSectionUserData.bRecomputeTangent = ExistingSection.bRecomputeTangent;
-					SourceSectionUserData.RecomputeTangentsVertexMaskChannel = ExistingSection.RecomputeTangentsVertexMaskChannel;
-				};
+					{
+						//Set the user section data to reflect the existing settings
+						SourceSectionUserData.bCastShadow = ExistingSection.bCastShadow;
+						SourceSectionUserData.bVisibleInRayTracing = ExistingSection.bVisibleInRayTracing;
+						SourceSectionUserData.bRecomputeTangent = ExistingSection.bRecomputeTangent;
+						SourceSectionUserData.RecomputeTangentsVertexMaskChannel = ExistingSection.RecomputeTangentsVertexMaskChannel;
+					};
 
 				if (ExistingImportedMaterialSlotName != NAME_None)
 				{
@@ -3952,14 +3945,69 @@ void FLODUtilities::MatchImportedMaterials(FLODUtilities::FSkeletalMeshMatchImpo
 					}
 				}
 				else if (Parameters.SkeletalMesh->GetMaterials().IsValidIndex(ExistingCurrentMaterialIndex) &&
-						 ImportedMaterials[NewSectionIndex].Material.Get() == Parameters.SkeletalMesh->GetMaterials()[ExistingCurrentMaterialIndex].MaterialInterface.Get()) //Use material slot compare to match in case the name is none
+					ImportedMaterials[NewSectionIndex].Material.Get() == Parameters.SkeletalMesh->GetMaterials()[ExistingCurrentMaterialIndex].MaterialInterface.Get()) //Use material slot compare to match in case the name is none
 				{
 					CopySectionData();
 					break;
 				}
 			}
 		}
+	}
+
+	LodInfo->LODMaterialMap.Empty();
+	// Now set up the material mapping array.
+	if (Parameters.LodIndex != 0)
+	{
+		for (int32 ImportedMaterialIndex = 0; ImportedMaterialIndex < ImportedMaterials.Num(); ImportedMaterialIndex++)
+		{
+			int32 LODMatIndex = INDEX_NONE;
+			FName ImportedMaterialName = *(ImportedMaterials[ImportedMaterialIndex].MaterialImportName);
+
+			//If the current lod section material slot was edited, we must keep the edition.
+			if (int32* MaterialRedirection = LODMaterialMapRedirection.Find(ImportedMaterialName))
+			{
+				LODMatIndex = *MaterialRedirection;
+			}
+			else
+			{
+				//Match by name
+				for (int32 MaterialIndex = 0; MaterialIndex < Materials.Num(); ++MaterialIndex)
+				{
+					const FSkeletalMaterial& SkeletalMaterial = Materials[MaterialIndex];
+					if (SkeletalMaterial.ImportedMaterialSlotName != NAME_None && SkeletalMaterial.ImportedMaterialSlotName == ImportedMaterialName)
+					{
+						LODMatIndex = MaterialIndex;
+						break;
+					}
+				}
+
+				// If we don't have a match, add a new entry to the material list.
+				if (LODMatIndex == INDEX_NONE)
+				{
+					LODMatIndex = Materials.Add(FSkeletalMaterial(ImportedMaterials[ImportedMaterialIndex].Material.Get(), true, false, ImportedMaterialName, ImportedMaterialName));
+				}
+			}
+			LodInfo->LODMaterialMap.Add(LODMatIndex);
+		}
+	}
+	else if (LODMaterialMapRedirection.Num() > 0)
+	{
+		for (int32 ExistMaterialIndex = 0; ExistMaterialIndex < Materials.Num(); ExistMaterialIndex++)
+		{
+			int32 LODMatIndex = INDEX_NONE;
+			FName ExistMaterialName = Materials[ExistMaterialIndex].ImportedMaterialSlotName;
+			//If the current lod section material slot was edited, we must keep the edition.
+			if (int32* MaterialRedirection = LODMaterialMapRedirection.Find(ExistMaterialName))
+			{
+				LODMatIndex = *MaterialRedirection;
+			}
+			LodInfo->LODMaterialMap.Add(LODMatIndex);
+		}
+	}
 
+	if (Parameters.bIsReImport)
+	{
+		FSkeletalMeshLODModel& ExistLODModel = Parameters.SkeletalMesh->GetImportedModel()->LODModels[Parameters.LodIndex];
 		for (int32 NewSectionIndex = 0; NewSectionIndex < ImportedMaterials.Num(); NewSectionIndex++)
 		{
 			FName ImportedMaterialName = *(ImportedMaterials[NewSectionIndex].MaterialImportName);
@@ -3972,7 +4020,7 @@ void FLODUtilities::MatchImportedMaterials(FLODUtilities::FSkeletalMeshMatchImpo
 		if (Parameters.LodIndex == 0)
 		{
 			//Since LOD 0 data will be re-arrange to follow the Material list order, sort the map accordingly.
-			NewUserSectionsDataMap.KeySort([Materials](const FName& A, const FName& B)
+			NewUserSectionsDataMap.KeySort([&Materials](const FName& A, const FName& B)
 				{
 					int32 MatAIndex = INDEX_NONE;
 					int32 MatBIndex = INDEX_NONE;
@@ -3998,88 +4046,48 @@ void FLODUtilities::MatchImportedMaterials(FLODUtilities::FSkeletalMeshMatchImpo
 					return MatAIndex < MatBIndex;
 				});
 		}
-		auto ApplyUserSectionLODModel = [&NewUserSectionsDataMap](FSkeletalMeshLODModel& ToApplyUserSectionLODModel, bool bSynchronizeUserSectionData)
+
+		//If the caller provide a CustomImportedLODModel we must apply the user section data to the provided lod model.
+		if (Parameters.CustomImportedLODModel)
 		{
-			ToApplyUserSectionLODModel.UserSectionsData.Reset();
+			Parameters.CustomImportedLODModel->UserSectionsData.Reset();
 			int32 RemapSectionIndex = 0;
 			for (TPair<FName, FSkelMeshSourceSectionUserData> UserSectionsData : NewUserSectionsDataMap)
 			{
-				ToApplyUserSectionLODModel.UserSectionsData.Add(RemapSectionIndex++, UserSectionsData.Value);
+				Parameters.CustomImportedLODModel->UserSectionsData.Add(RemapSectionIndex++, UserSectionsData.Value);
 			}
-			if (bSynchronizeUserSectionData)
-			{
-				ToApplyUserSectionLODModel.SyncronizeUserSectionsDataArray(true);
-			}
-		};
-
-		//If the caller provide a CustomImportedLODModel we must apply the user section data to the provided lod model.
-		if (Parameters.CustomImportedLODModel)
-		{
-			constexpr bool bSynchronizeUserSectionData = true;
-			ApplyUserSectionLODModel(*Parameters.CustomImportedLODModel, bSynchronizeUserSectionData);
+			Parameters.CustomImportedLODModel->SyncronizeUserSectionsDataArray(true);
 		}
 		else
 		{
-			constexpr bool bSynchronizeUserSectionData = false;
-			ApplyUserSectionLODModel(ExistLODModel, bSynchronizeUserSectionData);
-		}
-		
-	}
-
-	LodInfo->LODMaterialMap.Empty();
-	// Now set up the material mapping array.
-	if (Parameters.LodIndex != 0)
-	{
-		for (int32 ImportedMaterialIndex = 0; ImportedMaterialIndex < ImportedMaterials.Num(); ImportedMaterialIndex++)
-		{
-			int32 LODMatIndex = INDEX_NONE;
-			FName ImportedMaterialName = *(ImportedMaterials[ImportedMaterialIndex].MaterialImportName);
-
-			//If the current lod section material slot was edited, we must keep the edition.
-			if (int32* MaterialRedirection = LODMaterialMapRedirection.Find(ImportedMaterialName))
-			{
-				LODMatIndex = *MaterialRedirection;
-			}
-			else
+			ExistLODModel.UserSectionsData.Reset();
+			for (TPair<FName, FSkelMeshSourceSectionUserData> UserSectionsData : NewUserSectionsDataMap)
 			{
-				//Match by name
 				for (int32 MaterialIndex = 0; MaterialIndex < Materials.Num(); ++MaterialIndex)
 				{
-					const FSkeletalMaterial& SkeletalMaterial = Materials[MaterialIndex];
-					if (SkeletalMaterial.ImportedMaterialSlotName != NAME_None && SkeletalMaterial.ImportedMaterialSlotName == ImportedMaterialName)
+					if (Materials[MaterialIndex].ImportedMaterialSlotName == UserSectionsData.Key)
 					{
-						LODMatIndex = MaterialIndex;
+						ExistLODModel.UserSectionsData.Add(MaterialIndex, UserSectionsData.Value);
 						break;
 					}
 				}
-
-				// If we don't have a match, add a new entry to the material list.
-				if (LODMatIndex == INDEX_NONE)
-				{
-					LODMatIndex = Materials.Add(FSkeletalMaterial(ImportedMaterials[ImportedMaterialIndex].Material.Get(), true, false, ImportedMaterialName, ImportedMaterialName));
-				}
 			}
-			LodInfo->LODMaterialMap.Add(LODMatIndex);
+			//Reorder the user section from the keys
+			ExistLODModel.UserSectionsData.KeySort([](const int32& A, const int32& B)
+				{
+					return A < B;
+				});
 		}
+		
 	}
-	else if (LODMaterialMapRedirection.Num() > 0)
+
+	if (Parameters.bIsReImport)
 	{
-		for (int32 ExistMaterialIndex = 0; ExistMaterialIndex < Materials.Num(); ExistMaterialIndex++)
+		if(Parameters.LodIndex == 0)
 		{
-			int32 LODMatIndex = INDEX_NONE;
-			FName ExistMaterialName = Materials[ExistMaterialIndex].ImportedMaterialSlotName;
-			//If the current lod section material slot was edited, we must keep the edition.
-			if (int32* MaterialRedirection = LODMaterialMapRedirection.Find(ExistMaterialName))
-			{
-				LODMatIndex = *MaterialRedirection;
-			}
-			LodInfo->LODMaterialMap.Add(LODMatIndex);
+			ReorderMaterialSlotToBaseLod(Parameters.SkeletalMesh);
 		}
-	}
-
-	if (Parameters.bIsReImport && Parameters.LodIndex == 0 && GInterchangeSkeletalMeshReorderMaterialSlots)
-	{
-		ReorderMaterialSlotToBaseLod(Parameters.SkeletalMesh);
+		RemoveUnusedMaterialSlot(Parameters);
 	}
 }
 
@@ -4181,11 +4189,18 @@ void FLODUtilities::ReorderMaterialSlotToBaseLod(USkeletalMesh* SkeletalMesh)
 					RemapSectionsDataMap.Add(MaterialSlotRemap[UserSectionDataPair.Key], UserSectionDataPair.Value);
 				}
 				//Sort the remap section so its in the proper order. Optional but easier to debug ordered data
-				RemapSectionsDataMap.KeySort([Materials](const int32& A, const int32& B)
+				RemapSectionsDataMap.KeySort([](const int32& A, const int32& B)
 					{
 						return A < B;
 					});
-				LODModel.UserSectionsData = RemapSectionsDataMap;
+
+				//Rebuild a remap with the new order starting from 0
+				LODModel.UserSectionsData.Empty(RemapSectionsDataMap.Num());
+				int32 KeyIndex = 0;
+				for (TPair<int32, FSkelMeshSourceSectionUserData>& UserSectionDataPair : RemapSectionsDataMap)
+				{
+					LODModel.UserSectionsData.Add(KeyIndex++, UserSectionDataPair.Value);
+				}
 			}
 
 			//Remap the built sections material index, note that the sections should be rebuild after changing the material slot order. so not a critical step.
@@ -4221,6 +4236,64 @@ void FLODUtilities::ReorderMaterialSlotToBaseLod(USkeletalMesh* SkeletalMesh)
 	}
 }
 
+void FLODUtilities::RemoveUnusedMaterialSlot(FSkeletalMeshMatchImportedMaterialsParameters& Parameters)
+{
+	USkeletalMesh* SkeletalMesh = Parameters.SkeletalMesh;
+	if (!SkeletalMesh || !SkeletalMesh->IsLODImportedDataBuildAvailable(0))
+	{
+		return;
+	}
+
+	TArray<FSkeletalMaterial>& Materials = SkeletalMesh->GetMaterials();
+	if (Materials.Num() < 2)
+	{
+		return;
+	}
+	FSkeletalMeshModel* ImportedResource = SkeletalMesh->GetImportedModel();
+	const TArray<FSkeletalMeshLODInfo>& LODInfoArray = SkeletalMesh->GetLODInfoArray();
+
+	TArray<int32> UsedIndexes;
+	for (int32 LodIndex = 0; LodIndex < LODInfoArray.Num(); ++LodIndex)
+	{
+		const TArray<int32>& LODMaterialMap = LODInfoArray[LodIndex].LODMaterialMap;
+		if (!LODMaterialMap.IsEmpty())
+		{
+			for (int32 MaterialIndex : LODMaterialMap)
+			{
+				if (MaterialIndex != INDEX_NONE)
+				{
+					UsedIndexes.AddUnique(MaterialIndex);
+				}
+			}
+		}
+
+		if (LodIndex == 0 && ImportedResource && ImportedResource->LODModels.IsValidIndex(LodIndex))
+		{
+			//Since LOD 0 is always reordering the material slot array, we expect section index to match UserSectionsData order
+			//We cannot use the section since section are built by the DDC cache.
+			FSkeletalMeshLODModel& Model = ImportedResource->LODModels[LodIndex];
+			for (TPair<int32, FSkelMeshSourceSectionUserData> UserSectionsData : Model.UserSectionsData)
+			{
+				if (LODMaterialMap.IsEmpty() || !LODMaterialMap.IsValidIndex(UserSectionsData.Key) || LODMaterialMap[UserSectionsData.Key] == INDEX_NONE)
+				{
+					UsedIndexes.AddUnique(UserSectionsData.Key);
+				}
+			}
+		}
+	}
+
+	//Clean up the trailing unused material
+	for (int32 MaterialIndex = Materials.Num() - 1; MaterialIndex >= 0; MaterialIndex--)
+	{
+		if (UsedIndexes.Contains(MaterialIndex))
+		{
+			//Only delete extra material at the end of the material slot list, because MaterialMap and UserSectionData will need to be patch otherwise
+			break;
+		}
+		Materials.RemoveAt(MaterialIndex);
+	}
+}
+
 namespace TriangleStripHelper
 {
 	struct FTriangle2D
diff --git a/Engine/Source/Developer/SkeletalMeshUtilitiesCommon/Public/LODUtilities.h b/Engine/Source/Developer/SkeletalMeshUtilitiesCommon/Public/LODUtilities.h
--- a/Engine/Source/Developer/SkeletalMeshUtilitiesCommon/Public/LODUtilities.h
+++ b/Engine/Source/Developer/SkeletalMeshUtilitiesCommon/Public/LODUtilities.h
@@ -268,6 +268,11 @@ public:
 	 */
 	static void ReorderMaterialSlotToBaseLod(USkeletalMesh* SkeletalMesh);
 
+	/**
+	 * Remove any material slot that is not used by any LODs
+	 */
+	static void RemoveUnusedMaterialSlot(FSkeletalMeshMatchImportedMaterialsParameters& Parameters);
+
 	/**
 	 * This function will strip all triangle in the specified LOD that don't have any UV area pointing on a black pixel in the TextureMask.
 	 * We use the UVChannel 0 to find the pixels in the texture.
diff --git a/Engine/Source/Runtime/Engine/Classes/Engine/StaticMesh.h b/Engine/Source/Runtime/Engine/Classes/Engine/StaticMesh.h
--- a/Engine/Source/Runtime/Engine/Classes/Engine/StaticMesh.h
+++ b/Engine/Source/Runtime/Engine/Classes/Engine/StaticMesh.h
@@ -1602,6 +1602,13 @@ public:
 	 */
 	ENGINE_API bool SetCustomLOD(const UStaticMesh* SourceStaticMesh, int32 LodIndex, const FString& SourceDataFilename);
 
+	/*
+	 * Static function that remove any trailing unused material.
+	 *
+	 * @Param StaticMesh - The static mesh we want to remove the trailing materials.
+	 */
+	ENGINE_API static void RemoveUnusedMaterialSlots(UStaticMesh* StaticMesh);
+
 	//SourceModels API
 	ENGINE_API FStaticMeshSourceModel& AddSourceModel();
 
diff --git a/Engine/Source/Runtime/Engine/Private/StaticMesh.cpp b/Engine/Source/Runtime/Engine/Private/StaticMesh.cpp
--- a/Engine/Source/Runtime/Engine/Private/StaticMesh.cpp
+++ b/Engine/Source/Runtime/Engine/Private/StaticMesh.cpp
@@ -4249,7 +4249,6 @@ FStaticMeshSourceModel& UStaticMesh::GetSourceModel(int32 Index)
 
 bool UStaticMesh::SetCustomLOD(const UStaticMesh* SourceStaticMesh, int32 DestinationLodIndex, const FString& SourceDataFilename)
 {
-
 #if WITH_EDITORONLY_DATA
 	if (!ensure(SourceStaticMesh) || SourceStaticMesh->GetNumSourceModels() <= 0)
 	{
@@ -4393,12 +4392,63 @@ bool UStaticMesh::SetCustomLOD(const UStaticMesh* SourceStaticMesh, int32 Destin
 		SourceModel.SourceImportFilename = UAssetImportData::SanitizeImportFilename(SourceDataFilename, nullptr);
 		SourceModel.bImportWithBaseMesh = false;
 	}
+
+	RemoveUnusedMaterialSlots(this);
+
 	return true;
 #else
 	return false;
 #endif //!WITH_EDITORONLY_DATA
 }
 
+//Static function
+void UStaticMesh::RemoveUnusedMaterialSlots(UStaticMesh* StaticMesh)
+{
+#if WITH_EDITOR
+	if (!StaticMesh)
+	{
+		return;
+	}
+
+	TArray<FStaticMaterial>& Materials = StaticMesh->GetStaticMaterials();
+	if (Materials.Num() < 2)
+	{
+		return;
+	}
+
+	FMeshSectionInfoMap& SectionInfoMap = StaticMesh->GetSectionInfoMap();
+	const int32 LodCount = StaticMesh->GetNumSourceModels();
+
+	//Clean up the material list by removing unused materials
+	for (int32 MaterialIndex = Materials.Num() - 1; MaterialIndex >= 0; MaterialIndex--)
+	{
+		bool bMaterialIsUse = false;
+		for (int32 LodIndex = 0; LodIndex < LodCount && !bMaterialIsUse; ++LodIndex)
+		{
+			const int32 SectionCount = SectionInfoMap.GetSectionNumber(LodIndex);
+			for (int32 SectionIndex = 0; SectionIndex < SectionCount; ++SectionIndex)
+			{
+				FMeshSectionInfo SectionInfo = SectionInfoMap.Get(LodIndex, SectionIndex);
+				if (SectionInfo.MaterialIndex == MaterialIndex)
+				{
+					bMaterialIsUse = true;
+					break;
+				}
+			}
+		}
+		if (!bMaterialIsUse)
+		{
+			Materials.RemoveAt(MaterialIndex);
+		}
+		else
+		{
+			//Stop removing unused material when we find a valid one, to avoid patching any data related to material index.
+			break;
+		}
+	}
+#endif // WITH_EDITOR
+}
+
 FStaticMeshSourceModel& UStaticMesh::AddSourceModel()
 {
 	WaitUntilAsyncPropertyReleased(EStaticMeshAsyncProperties::SourceModels);
