AltSkin Import: When determining which points to process, ignore points that don't contribute to the mesh at all. FBX does not guarantee that all vertices are actually used by the mesh itself.

#rnx

diff --git a/Engine/Source/Developer/SkeletalMeshUtilitiesCommon/Private/LODUtilities.cpp b/Engine/Source/Developer/SkeletalMeshUtilitiesCommon/Private/LODUtilities.cpp
--- a/Engine/Source/Developer/SkeletalMeshUtilitiesCommon/Private/LODUtilities.cpp
+++ b/Engine/Source/Developer/SkeletalMeshUtilitiesCommon/Private/LODUtilities.cpp
@@ -2638,11 +2638,15 @@ bool FLODUtilities::UpdateAlternateSkinWeights(
 			}
 		}
 	}
-
+	
 	// Check if there are any points that have no weight assignments after all is said and done. Attempt to average across neighbors to reconstruct
 	// those weights. Currently this relies on these points being completely surrounded by at least two points that have weights assigned to them.
-	TBitArray<> UnassignedVertexBitmask(true, ImportDataDest.Points.Num());
-	for (SkeletalMeshImportData::FRawBoneInfluence& Influence: AlternateInfluences)
+	TBitArray<> UnassignedVertexBitmask(false, ImportDataDest.Points.Num());
+	for (const SkeletalMeshImportData::FVertex& Wedge: ImportDataDest.Wedges)
+	{
+		UnassignedVertexBitmask[Wedge.VertexIndex] = true;
+	}
+	for (const SkeletalMeshImportData::FRawBoneInfluence& Influence: AlternateInfluences)
 	{
 		UnassignedVertexBitmask[Influence.VertexIndex] = false;
 	}
