Interchange: bKeepSectionsSeparate LODs Fix

#rnx

diff --git a/Engine/Plugins/Interchange/Runtime/Source/Import/Private/Mesh/InterchangeSkeletalMeshFactory.cpp b/Engine/Plugins/Interchange/Runtime/Source/Import/Private/Mesh/InterchangeSkeletalMeshFactory.cpp
--- a/Engine/Plugins/Interchange/Runtime/Source/Import/Private/Mesh/InterchangeSkeletalMeshFactory.cpp
+++ b/Engine/Plugins/Interchange/Runtime/Source/Import/Private/Mesh/InterchangeSkeletalMeshFactory.cpp
@@ -1698,20 +1698,27 @@ UInterchangeFactoryBase::FImportAssetResult UInterchangeSkeletalMeshFactory::End
 					Parameters.ImportedMaterials = &ImportAssetObjectLODData.ImportedMaterials;
 					Parameters.ExistingOriginalPerSectionMaterialImportName = &ImportAssetObjectLODData.ExistingOriginalPerSectionMaterialImportName;
 					FLODUtilities::MatchImportedMaterials(Parameters);
-					//Flush the old LOD 0 sections after we rematch the materials
-					if (Parameters.LodIndex == 0 && Parameters.bIsReImport)
+					//Flush the old LOD sections after we rematch the materials
+					if (Parameters.bIsReImport)
 					{
 						if (FSkeletalMeshModel* ImportedModel = SkeletalMesh->GetImportedModel())
 						{
-							if (ImportedModel->LODModels.IsValidIndex(0))
+							if (ImportedModel->LODModels.IsValidIndex(ImportAssetObjectLODData.LodIndex))
 							{
-								SkeletalMesh->GetImportedModel()->LODModels[0].Sections.Empty();
+								SkeletalMesh->GetImportedModel()->LODModels[ImportAssetObjectLODData.LodIndex].Sections.Empty();
 							}
 						}
 					}
 				}
 			}
 		}
+
+		//Now that materials are matched we can re-order them and remove the unused.
+		if (ImportAssetObjectData.bIsReImport)
+		{
+			FLODUtilities::ReorderMaterialSlotToBaseLod(SkeletalMesh);
+			FLODUtilities::RemoveUnusedMaterialSlot(SkeletalMesh);
+		}
 	}
 
 	if (USkeleton* SkeletonReference = ImportAssetObjectData.SkeletonReference)
diff --git a/Engine/Plugins/Interchange/Runtime/Source/Import/Private/Mesh/InterchangeStaticMeshFactory.cpp b/Engine/Plugins/Interchange/Runtime/Source/Import/Private/Mesh/InterchangeStaticMeshFactory.cpp
--- a/Engine/Plugins/Interchange/Runtime/Source/Import/Private/Mesh/InterchangeStaticMeshFactory.cpp
+++ b/Engine/Plugins/Interchange/Runtime/Source/Import/Private/Mesh/InterchangeStaticMeshFactory.cpp
@@ -67,15 +67,6 @@ namespace UE::Interchange::Private::StaticMesh
 			return;
 		}
 
-		const FMeshDescription& BaseMeshDescription = *StaticMesh->GetMeshDescription(0);
-		if (!ensure(!BaseMeshDescription.NeedsCompact()))
-		{
-			return;
-		}
-		
-		FStaticMeshConstAttributes StaticMeshAttributes(BaseMeshDescription);
-		TPolygonGroupAttributesConstRef<FName> SlotNames = StaticMeshAttributes.GetPolygonGroupMaterialSlotNames();
-
 		TArray<int32> RemapMaterialIndexes;
 		RemapMaterialIndexes.Reserve(Materials.Num());
 		for (int32 MaterialIndex = 0; MaterialIndex < Materials.Num(); ++MaterialIndex)
@@ -84,30 +75,45 @@ namespace UE::Interchange::Private::StaticMesh
 		}
 		TArray<FStaticMaterial> ReorderMaterialArray;
 		ReorderMaterialArray.Reserve(Materials.Num());
-		for (FPolygonGroupID PolygonGroupID : BaseMeshDescription.PolygonGroups().GetElementIDs())
+
+		for (int32 LodIndex = 0; LodIndex < StaticMesh->GetNumSourceModels(); ++LodIndex)
 		{
-			FName ImportMaterialName = SlotNames[PolygonGroupID];
-			bool bFoundMatch = false;
-			for (int32 MaterialIndex = 0; MaterialIndex < Materials.Num(); ++MaterialIndex)
+			const FMeshDescription* LodMeshDescription = StaticMesh->GetMeshDescription(LodIndex);
+			if (!LodMeshDescription || !ensure(!LodMeshDescription->NeedsCompact()))
 			{
-				if (RemapMaterialIndexes[MaterialIndex] != INDEX_NONE)
+				if (LodIndex == 0)
 				{
-					//This material was already matched
-					continue;
+					return; //Lod 0 must always participate in the re-order, return if we can't use it
 				}
-				int32& RemapIndex = RemapMaterialIndexes[MaterialIndex];
-				FName MaterialName = Materials[MaterialIndex].ImportedMaterialSlotName;
-				if (MaterialName == ImportMaterialName)
+				continue;
+			}
+
+			FStaticMeshConstAttributes StaticMeshAttributes(*LodMeshDescription);
+			TPolygonGroupAttributesConstRef<FName> SlotNames = StaticMeshAttributes.GetPolygonGroupMaterialSlotNames();
+			for (FPolygonGroupID PolygonGroupID : LodMeshDescription->PolygonGroups().GetElementIDs())
+			{
+				FName ImportMaterialName = SlotNames[PolygonGroupID];
+				bool bFoundMatch = false;
+				for (int32 MaterialIndex = 0; MaterialIndex < Materials.Num(); ++MaterialIndex)
 				{
-					RemapIndex = ReorderMaterialArray.Add(Materials[MaterialIndex]);
-					bFoundMatch = true;
-					break;
+					if (RemapMaterialIndexes[MaterialIndex] != INDEX_NONE)
+					{
+						//This material was already matched
+						continue;
+					}
+					int32& RemapIndex = RemapMaterialIndexes[MaterialIndex];
+					FName MaterialName = Materials[MaterialIndex].ImportedMaterialSlotName;
+					if (MaterialName == ImportMaterialName)
+					{
+						RemapIndex = ReorderMaterialArray.Add(Materials[MaterialIndex]);
+						bFoundMatch = true;
+						break;
+					}
 				}
+				//All mesh description polygon group should have a match
+				ensure(bFoundMatch);
 			}
-			//All mesh description polygon group should have a match
-			ensure(bFoundMatch);
 		}
-
 		//Custom LOD can add materials, so we add them at the end of the material slots
 		for (int32 MaterialIndex = 0; MaterialIndex < Materials.Num(); ++MaterialIndex)
 		{
diff --git a/Engine/Plugins/Interchange/Runtime/Source/Pipelines/Public/InterchangePipelineMeshesUtilities.h b/Engine/Plugins/Interchange/Runtime/Source/Pipelines/Public/InterchangePipelineMeshesUtilities.h
--- a/Engine/Plugins/Interchange/Runtime/Source/Pipelines/Public/InterchangePipelineMeshesUtilities.h
+++ b/Engine/Plugins/Interchange/Runtime/Source/Pipelines/Public/InterchangePipelineMeshesUtilities.h
@@ -344,20 +344,23 @@ namespace UE::Interchange::MeshesUtilities
 		, TMap<FString, FString> *ExistingSlotMaterialDependenciesPtr)
 	{
 		bool bKeepSectionsSeparate = false;
+		int32 IndexCounter = 0; //Only use when bKeepSectionsSeparate is true
 		if (ExistingSlotMaterialDependenciesPtr)
 		{
 			FactoryNode.GetCustomKeepSectionsSeparate(bKeepSectionsSeparate);
+			IndexCounter = ExistingSlotMaterialDependenciesPtr->Num();
 		}
-
 		for (const TPair<FString, FString>& SlotMaterialDependency : SlotMaterialDependencies)
 		{
 			FString NewSlotName = SlotMaterialDependency.Key;
-			if (bKeepSectionsSeparate && ExistingSlotMaterialDependenciesPtr && ExistingSlotMaterialDependenciesPtr->Contains(NewSlotName))
+			if (bKeepSectionsSeparate && ExistingSlotMaterialDependenciesPtr)
 			{
-				TMap<FString, FString> NodeMaterialDependencies;
-				FactoryNode.GetSlotMaterialDependencies(NodeMaterialDependencies);
-				NewSlotName += TEXT("_Section") + FString::FromInt(NodeMaterialDependencies.Num());
+				if (ExistingSlotMaterialDependenciesPtr->Contains(NewSlotName))
+				{
+					NewSlotName += TEXT("_Section") + FString::FromInt(IndexCounter);
+				}
 				ExistingSlotMaterialDependenciesPtr->Add(NewSlotName, SlotMaterialDependency.Value);
+				IndexCounter++;
 			}
 			const FString MaterialFactoryNodeUid = UInterchangeBaseMaterialFactoryNode::GetMaterialFactoryNodeUidFromMaterialNodeUid(SlotMaterialDependency.Value);
 			FactoryNode.SetSlotMaterialDependencyUid(NewSlotName, MaterialFactoryNodeUid);
diff --git a/Engine/Source/Developer/SkeletalMeshUtilitiesCommon/Private/LODUtilities.cpp b/Engine/Source/Developer/SkeletalMeshUtilitiesCommon/Private/LODUtilities.cpp
--- a/Engine/Source/Developer/SkeletalMeshUtilitiesCommon/Private/LODUtilities.cpp
+++ b/Engine/Source/Developer/SkeletalMeshUtilitiesCommon/Private/LODUtilities.cpp
@@ -4082,16 +4082,6 @@ void FLODUtilities::MatchImportedMaterials(FLODUtilities::FSkeletalMeshMatchImpo
 					return A < B;
 				});
 		}
-		
-	}
-
-	if (Parameters.bIsReImport)
-	{
-		if(Parameters.LodIndex == 0)
-		{
-			ReorderMaterialSlotToBaseLod(Parameters.SkeletalMesh);
-		}
-		RemoveUnusedMaterialSlot(Parameters);
 	}
 }
 
@@ -4103,79 +4093,74 @@ void FLODUtilities::ReorderMaterialSlotToBaseLod(USkeletalMesh* SkeletalMesh)
 		UE_ASSET_LOG(LogLODUtilities, Warning, SkeletalMesh, TEXT("FLODUtilities::ReorderMaterialSlotToBaseLod: Bad parameters, SkeletalMesh is null"));
 		return;
 	}
-	if (SkeletalMesh->IsLODImportedDataEmpty(0))
-	{
-		UE_ASSET_LOG(LogLODUtilities, Warning, SkeletalMesh, TEXT("FLODUtilities::ReorderMaterialSlotToBaseLod: Skeletal mesh invalid import data for LOD 0"));
-		return;
-	}
-	FSkeletalMeshImportData BaseLodImportData;
-	SkeletalMesh->LoadLODImportedData(0, BaseLodImportData);
-	const TArray<SkeletalMeshImportData::FMaterial>& ImportedMaterials = BaseLodImportData.Materials;
 	TArray<FSkeletalMaterial>& Materials = SkeletalMesh->GetMaterials();
-	if (Materials.Num() == 0)
+	if (Materials.Num() < 2)
 	{
 		return;
 	}
 
 	TArray<int32> MaterialSlotRemap;
 	MaterialSlotRemap.AddUninitialized(Materials.Num());
-	TBitArray<> AvailableIndexes;
-	AvailableIndexes.Init(true, Materials.Num());
-	//Find remap index for LOD 0 matching materials
 	for (int32 MaterialIndex = 0; MaterialIndex < Materials.Num(); ++MaterialIndex)
 	{
-		FName MaterialSlotName = Materials[MaterialIndex].ImportedMaterialSlotName;
 		MaterialSlotRemap[MaterialIndex] = INDEX_NONE;
-		for (int32 ImportedMaterialIndex = 0; ImportedMaterialIndex < ImportedMaterials.Num(); ++ImportedMaterialIndex)
+	}
+	TArray<FSkeletalMaterial> ReorderMaterialArray;
+	ReorderMaterialArray.Reserve(Materials.Num());
+
+	for (int32 LodIndex = 0; LodIndex < SkeletalMesh->GetLODNum(); ++LodIndex)
+	{
+		if (SkeletalMesh->IsLODImportedDataEmpty(LodIndex))
 		{
-			FName ImportedMaterialSlotName = FName(*ImportedMaterials[ImportedMaterialIndex].MaterialImportName);
-			if (MaterialSlotName == ImportedMaterialSlotName)
+			if (LodIndex == 0)
 			{
-				MaterialSlotRemap[MaterialIndex] = ImportedMaterialIndex;
-				if(ensure(AvailableIndexes.IsValidIndex(ImportedMaterialIndex) && AvailableIndexes[ImportedMaterialIndex]))
-				{
-					AvailableIndexes[ImportedMaterialIndex] = false;
-				}
-				break;
+				UE_ASSET_LOG(LogLODUtilities, Warning, SkeletalMesh, TEXT("FLODUtilities::ReorderMaterialSlotToBaseLod: Skeletal mesh invalid import data for LOD 0"));
+				return;
 			}
+			continue;
 		}
-	}
-	//Find remap index for any extra custom LOD materials
-	for (int32 MaterialIndex = 0; MaterialIndex < Materials.Num(); ++MaterialIndex)
-	{
-		if (MaterialSlotRemap[MaterialIndex] == INDEX_NONE)
+
+		FSkeletalMeshImportData LodImportData;
+		SkeletalMesh->LoadLODImportedData(LodIndex, LodImportData);
+		const TArray<SkeletalMeshImportData::FMaterial>& ImportedMaterials = LodImportData.Materials;
+
+		//Find remap index for this LOD matching materials
+		for (int32 ImportedMaterialIndex = 0; ImportedMaterialIndex < ImportedMaterials.Num(); ++ImportedMaterialIndex)
 		{
-			//Give the first available material index
-			for (int32 AvailableIndex = 0; AvailableIndex < AvailableIndexes.Num(); ++AvailableIndex)
+			FName ImportedMaterialSlotName = FName(*ImportedMaterials[ImportedMaterialIndex].MaterialImportName);
+			bool bFoundMatch = false;
+			for (int32 MaterialIndex = 0; MaterialIndex < Materials.Num(); ++MaterialIndex)
 			{
-				if (AvailableIndexes[AvailableIndex])
+				if (MaterialSlotRemap[MaterialIndex] != INDEX_NONE)
 				{
-					MaterialSlotRemap[MaterialIndex] = AvailableIndex;
-					AvailableIndexes[AvailableIndex] = false;
+					continue;
+				}
+				int32& RemapIndex = MaterialSlotRemap[MaterialIndex];
+				FName MaterialSlotName = Materials[MaterialIndex].ImportedMaterialSlotName;
+				if (MaterialSlotName == ImportedMaterialSlotName)
+				{
+					RemapIndex = ReorderMaterialArray.Add(Materials[MaterialIndex]);
+					bFoundMatch = true;
 					break;
 				}
 			}
-			//We should always be able to rematch a material
-			ensure(MaterialSlotRemap[MaterialIndex] != INDEX_NONE);
+			//All mesh description polygon group should have a match
+			ensure(bFoundMatch);
 		}
 	}
-
-	//Build the re-ordered maps
-	TArray<FSkeletalMaterial> ReorderMaterials;
-	ReorderMaterials.Reserve(Materials.Num());
+	//Custom LOD can add materials, so we add them at the end of the material slots
 	for (int32 MaterialIndex = 0; MaterialIndex < Materials.Num(); ++MaterialIndex)
 	{
-		for (int32 RemapIndex = 0; RemapIndex < MaterialSlotRemap.Num(); ++RemapIndex)
+		if (MaterialSlotRemap[MaterialIndex] == INDEX_NONE)
 		{
-			if (MaterialSlotRemap[RemapIndex] == MaterialIndex)
-			{
-				ReorderMaterials.Add(Materials[RemapIndex]);
-			}
+			MaterialSlotRemap[MaterialIndex] = ReorderMaterialArray.Add(Materials[MaterialIndex]);
 		}
 	}
 
-	//Change the material slot array
-	SkeletalMesh->SetMaterials(ReorderMaterials);
+	check(ReorderMaterialArray.Num() == Materials.Num());
+
+	//Reorder the skeletal mesh material slot array
+	SkeletalMesh->SetMaterials(ReorderMaterialArray);
 	
 	//We now need to adjust all LODs data to fit the re-order
 	for (int32 LodIndex = 0; LodIndex < SkeletalMesh->GetLODNum(); ++LodIndex)
@@ -4183,38 +4168,23 @@ void FLODUtilities::ReorderMaterialSlotToBaseLod(USkeletalMesh* SkeletalMesh)
 		if (ensure(SkeletalMesh->GetImportedModel()->LODModels.IsValidIndex(LodIndex)))
 		{
 			FSkeletalMeshLODModel& LODModel = SkeletalMesh->GetImportedModel()->LODModels[LodIndex];
-			//The skeletal mesh build process have a different behavior for LOD 0. The build reorder the import data sections in the material slot order.
-			//Because of this we need to remap the UserSectionData key
-			if (LodIndex == 0)
+			TMap<int32, FSkelMeshSourceSectionUserData> RemapSectionsDataMap;
+			for (TPair<int32, FSkelMeshSourceSectionUserData>& UserSectionDataPair : LODModel.UserSectionsData)
 			{
-				TMap<int32, FSkelMeshSourceSectionUserData> RemapSectionsDataMap;
-				for (TPair<int32, FSkelMeshSourceSectionUserData>& UserSectionDataPair : LODModel.UserSectionsData)
-				{
-					RemapSectionsDataMap.Add(MaterialSlotRemap[UserSectionDataPair.Key], UserSectionDataPair.Value);
-				}
-				//Sort the remap section so its in the proper order. Optional but easier to debug ordered data
-				RemapSectionsDataMap.KeySort([](const int32& A, const int32& B)
-					{
-						return A < B;
-					});
-
-				//Rebuild a remap with the new order starting from 0
-				LODModel.UserSectionsData.Empty(RemapSectionsDataMap.Num());
-				int32 KeyIndex = 0;
-				for (TPair<int32, FSkelMeshSourceSectionUserData>& UserSectionDataPair : RemapSectionsDataMap)
-				{
-					LODModel.UserSectionsData.Add(KeyIndex++, UserSectionDataPair.Value);
-				}
+				RemapSectionsDataMap.Add(MaterialSlotRemap[UserSectionDataPair.Key], UserSectionDataPair.Value);
 			}
+			//Sort the remap section so its in the proper order. Optional but easier to debug ordered data
+			RemapSectionsDataMap.KeySort([](const int32& A, const int32& B)
+				{
+					return A < B;
+				});
 
-			//Remap the built sections material index, note that the sections should be rebuild after changing the material slot order. so not a critical step.
-			for (int32 SectionIndex = 0; SectionIndex < LODModel.Sections.Num(); ++SectionIndex)
+			//Rebuild a remap with the new order starting from 0 to follow the section order
+			LODModel.UserSectionsData.Empty(RemapSectionsDataMap.Num());
+			int32 KeyIndex = 0;
+			for (TPair<int32, FSkelMeshSourceSectionUserData>& UserSectionDataPair : RemapSectionsDataMap)
 			{
-				FSkelMeshSection& Section = LODModel.Sections[SectionIndex];
-				if(MaterialSlotRemap.IsValidIndex(Section.MaterialIndex))
-				{
-					Section.MaterialIndex = MaterialSlotRemap[Section.MaterialIndex];
-				}
+				LODModel.UserSectionsData.Add(KeyIndex++, UserSectionDataPair.Value);
 			}
 		}
 		else
@@ -4240,9 +4210,8 @@ void FLODUtilities::ReorderMaterialSlotToBaseLod(USkeletalMesh* SkeletalMesh)
 	}
 }
 
-void FLODUtilities::RemoveUnusedMaterialSlot(FSkeletalMeshMatchImportedMaterialsParameters& Parameters)
+void FLODUtilities::RemoveUnusedMaterialSlot(USkeletalMesh* SkeletalMesh)
 {
-	USkeletalMesh* SkeletalMesh = Parameters.SkeletalMesh;
 	if (!SkeletalMesh || !SkeletalMesh->IsLODImportedDataBuildAvailable(0))
 	{
 		return;
diff --git a/Engine/Source/Developer/SkeletalMeshUtilitiesCommon/Public/LODUtilities.h b/Engine/Source/Developer/SkeletalMeshUtilitiesCommon/Public/LODUtilities.h
--- a/Engine/Source/Developer/SkeletalMeshUtilitiesCommon/Public/LODUtilities.h
+++ b/Engine/Source/Developer/SkeletalMeshUtilitiesCommon/Public/LODUtilities.h
@@ -271,7 +271,7 @@ public:
 	/**
 	 * Remove any material slot that is not used by any LODs
 	 */
-	static void RemoveUnusedMaterialSlot(FSkeletalMeshMatchImportedMaterialsParameters& Parameters);
+	static void RemoveUnusedMaterialSlot(USkeletalMesh* SkeletalMesh);
 
 	/**
 	 * This function will strip all triangle in the specified LOD that don't have any UV area pointing on a black pixel in the TextureMask.
