Skeletal Mesh: No need to warn on missing mesh description when processing vertex attributes since we might be processing a dependent LOD.

#rnx

diff --git a/Engine/Source/Developer/SkeletalMeshUtilitiesCommon/Private/LODUtilities.cpp b/Engine/Source/Developer/SkeletalMeshUtilitiesCommon/Private/LODUtilities.cpp
--- a/Engine/Source/Developer/SkeletalMeshUtilitiesCommon/Private/LODUtilities.cpp
+++ b/Engine/Source/Developer/SkeletalMeshUtilitiesCommon/Private/LODUtilities.cpp
@@ -2978,7 +2978,7 @@ bool FLODUtilities::UpdateLODInfoVertexAttributes(
 	bool bInCopyAttributeValues
 	)
 {
-	if (!ensure(InSkeletalMesh->HasMeshDescription(InSourceLODIndex)))
+	if (!InSkeletalMesh->HasMeshDescription(InSourceLODIndex))
 	{
 		return false;
 	}
