[Skeletal Mesh Utilities] fixed a bug causing the same vert attribute map getting copied into multiple attributes on the LOD model

diff --git a/Engine/Source/Developer/SkeletalMeshUtilitiesCommon/Private/LODUtilities.cpp b/Engine/Source/Developer/SkeletalMeshUtilitiesCommon/Private/LODUtilities.cpp
--- a/Engine/Source/Developer/SkeletalMeshUtilitiesCommon/Private/LODUtilities.cpp
+++ b/Engine/Source/Developer/SkeletalMeshUtilitiesCommon/Private/LODUtilities.cpp
@@ -3015,9 +3015,9 @@ bool FLODUtilities::UpdateLODInfoVertexAttributes(
 	
 	TArray<UE::Tasks::FTask> ConversionTasks;
 	
-	for (const TPair<FName, TVertexAttributesConstRef<float>> SourceAttributeInfo: SourceAttributes)
+	for (const TPair<FName, TVertexAttributesConstRef<float>>& SourceAttributeInfo: SourceAttributes)
 	{
-		const TVertexAttributesConstRef<float>& SourceAttribute = SourceAttributeInfo.Value;
+		const TVertexAttributesConstRef<float> SourceAttribute = SourceAttributeInfo.Value;
 		const FName AttributeName(SourceAttributeInfo.Key);
 
 		// Did this definition already exist? Try to retain as much of the existing information as possible.
@@ -3044,8 +3044,9 @@ bool FLODUtilities::UpdateLODInfoVertexAttributes(
 			if (InTargetLODIndex == InSourceLODIndex)
 			{
 				ConversionTasks.Add(
-					UE::Tasks::Launch(UE_SOURCE_LOCATION, [&TargetLODModel, &ModelAttribute, &SourceAttribute]()
+					UE::Tasks::Launch(UE_SOURCE_LOCATION, [&TargetLODModel, AttributeName, SourceAttribute]()
 					{
+						FSkeletalMeshModelVertexAttribute& ModelAttribute = TargetLODModel.VertexAttributes.FindChecked(AttributeName);	
 						ModelAttribute.Values.SetNumUninitialized(TargetLODModel.NumVertices);
 						for(uint32 VertexIndex = 0; VertexIndex < TargetLODModel.NumVertices; VertexIndex++)
 						{
