Skeletal Mesh: Change storage of morph target normals to be on vertex instances rather than vertices, to match how the base normals are stored. Add a backward compatibility path for skeletal meshes imported from Alembic files that have manually constructed UMorphTargets to support non-standard normals on the morph targets.

#rnx

diff --git a/Engine/Source/Developer/MeshBuilder/Private/SkeletalMeshBuilder.cpp b/Engine/Source/Developer/MeshBuilder/Private/SkeletalMeshBuilder.cpp
--- a/Engine/Source/Developer/MeshBuilder/Private/SkeletalMeshBuilder.cpp
+++ b/Engine/Source/Developer/MeshBuilder/Private/SkeletalMeshBuilder.cpp
@@ -84,7 +84,6 @@ bool FSkeletalMeshBuilder::Build(const FSkeletalMeshBuildParameters& SkeletalMes
 	TArray<ClothingAssetUtils::FClothingAssetMeshBinding> ClothingBindings;
 	FLODUtilities::UnbindClothingAndBackup(SkeletalMesh, ClothingBindings, LODIndex);
 
-	FSkeletalMeshImportData SkeletalMeshImportData;
 	int32 NumTextCoord = 1; //We need to send rendering at least one tex coord buffer
 
 	//This scope define where we can use the LODModel, after a reduction the LODModel must be requery since it is a new instance
@@ -92,9 +91,8 @@ bool FSkeletalMeshBuilder::Build(const FSkeletalMeshBuildParameters& SkeletalMes
 		FSkeletalMeshLODModel& BuildLODModel = SkeletalMesh->GetImportedModel()->LODModels[LODIndex];
 
 		//Load the imported data
-		PRAGMA_DISABLE_DEPRECATION_WARNINGS
-		SkeletalMesh->LoadLODImportedData(LODIndex, SkeletalMeshImportData);
-		PRAGMA_ENABLE_DEPRECATION_WARNINGS
+		const FMeshDescription& SkeletalMeshModel = *SkeletalMesh->GetMeshDescription(LODIndex);
+		FSkeletalMeshImportData SkeletalMeshImportData = FSkeletalMeshImportData::CreateFromMeshDescription(SkeletalMeshModel);
 
 		TArray<FVector3f> LODPoints;
 		TArray<SkeletalMeshImportData::FMeshWedge> LODWedges;
@@ -153,7 +151,7 @@ bool FSkeletalMeshBuilder::Build(const FSkeletalMeshBuildParameters& SkeletalMes
 		SlowTask.EnterProgressFrame(1.0f, NSLOCTEXT("SkeltalMeshBuilder", "RebuildMorphTarget", "Rebuilding morph targets..."));
 		if (SkeletalMeshImportData.MorphTargetNames.Num() > 0)
 		{
-			FLODUtilities::BuildMorphTargets(SkeletalMesh, SkeletalMeshImportData, LODIndex, !Options.bComputeNormals, !Options.bComputeTangents, Options.bUseMikkTSpace, Options.OverlappingThresholds);
+			FLODUtilities::BuildMorphTargets(SkeletalMesh, SkeletalMeshModel, SkeletalMeshImportData, LODIndex, !Options.bComputeNormals, !Options.bComputeTangents, Options.bUseMikkTSpace, Options.OverlappingThresholds);
 		}
 
 		//Re-apply the alternate skinning it must be after the inline reduction
diff --git a/Engine/Source/Developer/SkeletalMeshUtilitiesCommon/Private/LODUtilities.cpp b/Engine/Source/Developer/SkeletalMeshUtilitiesCommon/Private/LODUtilities.cpp
--- a/Engine/Source/Developer/SkeletalMeshUtilitiesCommon/Private/LODUtilities.cpp
+++ b/Engine/Source/Developer/SkeletalMeshUtilitiesCommon/Private/LODUtilities.cpp
@@ -3300,6 +3300,7 @@ void FLODUtilities::RegenerateDependentLODs(USkeletalMesh* SkeletalMesh, int32 L
 	}
 }
 
+
 //////////////////////////////////////////////////////////////////////////
 // Morph targets build code
 //
@@ -3339,14 +3340,31 @@ static void ConvertImportDataToMeshData(const FSkeletalMeshImportData& ImportDat
 class FAsyncImportMorphTargetWork : public FNonAbandonableTask
 {
 public:
-	FAsyncImportMorphTargetWork(FSkeletalMeshLODModel* InLODModel, const FReferenceSkeleton& InRefSkeleton, const FSkeletalMeshImportData& InBaseImportData, TArray<FVector3f>&& InMorphLODPoints,
-		TArray< FMorphTargetDelta >& InMorphDeltas, TArray<uint32>& InBaseIndexData, const TArray< uint32 >& InBaseWedgePointIndices,
-		TMap<uint32, uint32>& InWedgePointToVertexIndexMap, const FOverlappingCorners& InOverlappingCorners,
-		const TSet<uint32> InModifiedPoints, const TMultiMap< int32, int32 >& InWedgeToFaces, const FMeshDataBundle& InMeshDataBundle, const TArray<FVector3f>& InTangentZ,
-		bool InShouldImportNormals, bool InShouldImportTangents, bool InbUseMikkTSpace, const FOverlappingThresholds InThresholds)
+	FAsyncImportMorphTargetWork(
+		const FName InMorphName,
+		FSkeletalMeshLODModel* InLODModel, 
+		const FReferenceSkeleton& InRefSkeleton,
+		const FSkeletalMeshImportData& InBaseImportData,
+		const FMeshDescription* InSkeletalMeshModel,
+		TArray<FVector3f>&& InMorphLODPoints,
+		TArray< FMorphTargetDelta >& InMorphDeltas,
+		TArray<uint32>& InBaseIndexData,
+		const TArray< uint32 >& InBaseWedgePointIndices,
+		TMap<uint32, uint32>& InWedgePointToVertexIndexMap,
+		const FOverlappingCorners& InOverlappingCorners,
+		const TSet<uint32>& InModifiedPoints,
+		const TMultiMap< int32, int32 >& InWedgeToFaces,
+		const FMeshDataBundle& InMeshDataBundle,
+		const TArray<FVector3f>& InTangentZ,
+		bool InShouldImportNormals,
+		bool InShouldImportTangents, 
+		bool InbUseMikkTSpace, 
+		const FOverlappingThresholds InThresholds
+		)
 		: LODModel(InLODModel)
 		, RefSkeleton(InRefSkeleton)
 		, BaseImportData(InBaseImportData)
+		, SkeletalMeshModel(InSkeletalMeshModel)
 		, CompressMorphLODPoints(InMorphLODPoints)
 		, MorphTargetDeltas(InMorphDeltas)
 		, BaseIndexData(InBaseIndexData)
@@ -3364,6 +3382,15 @@ public:
 		, Thresholds(InThresholds)
 	{
 		MeshUtilities = &FModuleManager::Get().LoadModuleChecked<IMeshUtilities>("MeshUtilities");
+
+		if (SkeletalMeshModel)
+		{
+			FSkeletalMeshConstAttributes SkeletalMeshAttributes(*SkeletalMeshModel);
+
+			BaseNormalAttribute = SkeletalMeshAttributes.GetVertexInstanceNormals();
+			MorphNormalDeltaAttribute = SkeletalMeshAttributes.GetVertexInstanceMorphNormalDelta(InMorphName);
+			ShouldImportNormals = false;
+		}
 	}
 
 	//Decompress the shape points data
@@ -3472,60 +3499,77 @@ public:
 		WasProcessed.Empty(LODModel->NumVertices);
 		WasProcessed.AddZeroed(LODModel->NumVertices);
 
+		const bool bUseAttributeForNormal = MorphNormalDeltaAttribute.IsValid();
+
 		for (int32 Idx = 0; Idx < BaseIndexData.Num(); ++Idx)
 		{
 			uint32 BaseVertIdx = BaseIndexData[Idx];
 			// check for duplicate processing
-			if (!WasProcessed[BaseVertIdx])
+			if (WasProcessed[BaseVertIdx])
 			{
-				// mark this base vertex as already processed
-				WasProcessed[BaseVertIdx] = true;
-
-				// clothing can add extra verts, and we won't have source point, so we ignore those
-				if (BaseWedgePointIndices.IsValidIndex(BaseVertIdx))
-				{
-					// get the base mesh's original wedge point index
-					uint32 BasePointIdx = BaseWedgePointIndices[BaseVertIdx];
-					if (MeshDataBundle.Vertices.IsValidIndex(BasePointIdx) && MorphLODPoints.IsValidIndex(BasePointIdx))
-					{
-						FVector BasePosition = (FVector)MeshDataBundle.Vertices[BasePointIdx];
-						FVector TargetPosition = (FVector)MorphLODPoints[BasePointIdx];
+				continue;
+			}
+			
+			// mark this base vertex as already processed
+			WasProcessed[BaseVertIdx] = true;
 
-						FVector PositionDelta = TargetPosition - BasePosition;
+			// clothing can add extra verts, and we won't have source point, so we ignore those
+			if (!BaseWedgePointIndices.IsValidIndex(BaseVertIdx))
+			{
+				continue;
+			}
+			
+			// get the base mesh's original wedge point index
+			uint32 BasePointIdx = BaseWedgePointIndices[BaseVertIdx];
+			if (MeshDataBundle.Vertices.IsValidIndex(BasePointIdx) && MorphLODPoints.IsValidIndex(BasePointIdx))
+			{
+				const FVector3f& BasePosition = MeshDataBundle.Vertices[BasePointIdx];
+				const FVector3f& TargetPosition = MorphLODPoints[BasePointIdx];
 
-						uint32* VertexIdx = WedgePointToVertexIndexMap.Find(BasePointIdx);
+				const FVector3f PositionDelta = TargetPosition - BasePosition;
 
-						FVector NormalDeltaZ = FVector::ZeroVector;
+				uint32* VertexIdx = WedgePointToVertexIndexMap.Find(BasePointIdx);
 
-						if (VertexIdx != nullptr)
-						{
-							FVector BaseNormal = (FVector)BaseTangentZ[*VertexIdx];
-							FVector TargetNormal = (FVector)TangentZ[*VertexIdx];
+				FVector3f NormalDeltaZ = FVector3f::ZeroVector;
 
-							NormalDeltaZ = TargetNormal - BaseNormal;
-						}
+				if (VertexIdx != nullptr)
+				{
+					FVector3f BaseNormal = BaseTangentZ[*VertexIdx];
 
-						// check if position actually changed much
-						if (PositionDelta.SizeSquared() > FMath::Square(Thresholds.MorphThresholdPosition) ||
-							// since we can't get imported morphtarget normal from FBX
-							// we can't compare normal unless it's calculated
-							// this is special flag to ignore normal diff
-							((ShouldImportNormals == false) && NormalDeltaZ.SizeSquared() > 0.01f))
-						{
-							// create a new entry
-							FMorphTargetDelta NewVertex;
-							// position delta
-							NewVertex.PositionDelta = (FVector3f)PositionDelta;
-							// normal delta
-							NewVertex.TangentZDelta = (FVector3f)NormalDeltaZ;
-							// index of base mesh vert this entry is to modify
-							NewVertex.SourceIdx = BaseVertIdx;
-
-							// add it to the list of changed verts
-							MorphTargetDeltas.Add(NewVertex);
-						}
+					if (bUseAttributeForNormal)
+					{
+						// BasePointIdx is the index into the import data points. These map directly to the mesh description's vertex ids.
+						const FVertexID VertexID(BasePointIdx);
+						const FVertexInstanceID VertexInstanceID = SkeletalMeshModel->GetVertexVertexInstanceIDs(VertexID)[0];
+						// Use the first one. They all carry the same information for now due to the only place we get these normals are from
+						// rebuilding from a FSkeletalMeshLODModel + UMorphTargets.
+						NormalDeltaZ = MorphNormalDeltaAttribute.Get(VertexInstanceID);
+					}
+					else
+					{
+						NormalDeltaZ = TangentZ[*VertexIdx] - BaseNormal;
 					}
 				}
+
+				// check if position actually changed much
+				if (PositionDelta.SizeSquared() > FMath::Square(Thresholds.MorphThresholdPosition) ||
+					// since we can't get imported morph target normal from FBX
+					// we can't compare normal unless it's calculated
+					// this is special flag to ignore normal diff
+					((ShouldImportNormals == false) && NormalDeltaZ.SizeSquared() > 0.01f))
+				{
+					// create a new entry
+					FMorphTargetDelta NewVertex;
+					// position delta
+					NewVertex.PositionDelta = PositionDelta;
+					// normal delta
+					NewVertex.TangentZDelta = NormalDeltaZ;
+					// index of base mesh vert this entry is to modify
+					NewVertex.SourceIdx = BaseVertIdx;
+
+					// add it to the list of changed verts
+					MorphTargetDeltas.Add(NewVertex);
+				}
 			}
 		}
 	}
@@ -3533,8 +3577,11 @@ public:
 	void DoWork()
 	{
 		DecompressData();
-		PrepareTangents();
-		ComputeTangents();
+		if (!MorphNormalDeltaAttribute.IsValid())
+		{
+			PrepareTangents();
+			ComputeTangents();
+		}
 		ComputeMorphDeltas();
 	}
 
@@ -3564,6 +3611,9 @@ private:
 	// @todo not thread safe
 	const FReferenceSkeleton& RefSkeleton;
 	const FSkeletalMeshImportData& BaseImportData;
+	const FMeshDescription* SkeletalMeshModel;
+	TVertexInstanceAttributesConstRef<FVector3f> BaseNormalAttribute;
+	TVertexInstanceAttributesConstRef<FVector3f> MorphNormalDeltaAttribute;
 	const TArray<FVector3f> CompressMorphLODPoints;
 	TArray<FVector3f> MorphLODPoints;
 
@@ -3575,7 +3625,7 @@ private:
 	TMap<uint32, uint32>& WedgePointToVertexIndexMap;
 
 	const FOverlappingCorners& OverlappingCorners;
-	const TSet<uint32> ModifiedPoints;
+	const TSet<uint32>& ModifiedPoints;
 	const TMultiMap< int32, int32 >& WedgeToFaces;
 	const FMeshDataBundle& MeshDataBundle;
 
@@ -3587,7 +3637,17 @@ private:
 	const FOverlappingThresholds Thresholds;
 };
 
-void FLODUtilities::BuildMorphTargets(USkeletalMesh* BaseSkelMesh, FSkeletalMeshImportData &BaseImportData, int32 LODIndex, bool ShouldImportNormals, bool ShouldImportTangents, bool bUseMikkTSpace, const FOverlappingThresholds& Thresholds)
+
+static void BuildMorphTargetsInternal(
+	USkeletalMesh* BaseSkelMesh,
+	const FMeshDescription* SkeletalMeshModel,
+	FSkeletalMeshImportData &BaseImportData,
+	int32 LODIndex,
+	bool ShouldImportNormals,
+	bool ShouldImportTangents,
+	bool bUseMikkTSpace,
+	const FOverlappingThresholds& Thresholds
+	)
 {
 	bool bComputeNormals = !ShouldImportNormals || !BaseImportData.bHasNormals;
 	bool bComputeTangents = !ShouldImportTangents || !BaseImportData.bHasTangents;
@@ -3742,8 +3802,8 @@ void FLODUtilities::BuildMorphTargets(USkeletalMesh* BaseSkelMesh, FSkeletalMesh
 
 			TArray< FMorphTargetDelta >* Deltas = Results[NewMorphDeltasIdx];
 
-			FAsyncTask<FAsyncImportMorphTargetWork>* NewWork = new FAsyncTask<FAsyncImportMorphTargetWork>(&BaseLODModel, BaseSkelMesh->GetRefSkeleton(), BaseImportData,
-				MoveTemp(ShapeImportData.Points), *Deltas, BaseIndexData, BaseLODModel.GetRawPointIndices(), WedgePointToVertexIndexMap, OverlappingVertices, MoveTemp(ModifiedPoints), WedgeToFaces, MeshDataBundle, TangentZ,
+			FAsyncTask<FAsyncImportMorphTargetWork>* NewWork = new FAsyncTask<FAsyncImportMorphTargetWork>(MorphTarget->GetFName(), &BaseLODModel, BaseSkelMesh->GetRefSkeleton(), BaseImportData, SkeletalMeshModel,
+				MoveTemp(ShapeImportData.Points), *Deltas, BaseIndexData, BaseLODModel.GetRawPointIndices(), WedgePointToVertexIndexMap, OverlappingVertices, ModifiedPoints, WedgeToFaces, MeshDataBundle, TangentZ,
 				ShouldImportNormals, ShouldImportTangents, bUseMikkTSpace, Thresholds);
 			PendingWork.Add(NewWork);
 
@@ -3786,7 +3846,9 @@ void FLODUtilities::BuildMorphTargets(USkeletalMesh* BaseSkelMesh, FSkeletalMesh
 
 		UMorphTarget* MorphTarget = MorphTargets[Index];
 		check(IsValid(MorphTarget));
-		MorphTarget->PopulateDeltas(*Results[Index], LODIndex, BaseLODModel.Sections, ShouldImportNormals == false, false, Thresholds.MorphThresholdPosition);
+		constexpr bool bCompareNormals = true;		// Ensure we include morphs that only modify normals and not positions. 
+		constexpr bool bGeneratedByReduction = false;
+		MorphTarget->PopulateDeltas(*Results[Index], LODIndex, BaseLODModel.Sections, bCompareNormals, bGeneratedByReduction, Thresholds.MorphThresholdPosition);
 
 		// register does mark package as dirty
 		if (MorphTarget->HasValidData())
@@ -3808,6 +3870,18 @@ void FLODUtilities::BuildMorphTargets(USkeletalMesh* BaseSkelMesh, FSkeletalMesh
 	}
 }
 
+void FLODUtilities::BuildMorphTargets(USkeletalMesh* SkeletalMesh, FSkeletalMeshImportData& ImportData, int32 LODIndex, bool ShouldImportNormals, bool ShouldImportTangents, bool bUseMikkTSpace, const FOverlappingThresholds& Thresholds)
+{
+	BuildMorphTargetsInternal(SkeletalMesh, nullptr, ImportData, LODIndex, ShouldImportNormals, ShouldImportTangents, bUseMikkTSpace, Thresholds);
+}
+
+void FLODUtilities::BuildMorphTargets(USkeletalMesh* SkeletalMesh, const FMeshDescription& SkeletalMeshModel, FSkeletalMeshImportData& ImportData, int32 LODIndex, bool ShouldImportNormals, bool ShouldImportTangents, bool bUseMikkTSpace, const FOverlappingThresholds& Thresholds)
+{
+	BuildMorphTargetsInternal(SkeletalMesh, &SkeletalMeshModel, ImportData, LODIndex, ShouldImportNormals, ShouldImportTangents, bUseMikkTSpace, Thresholds);
+}
+
+
+
 void FLODUtilities::UnbindClothingAndBackup(USkeletalMesh* SkeletalMesh, TArray<ClothingAssetUtils::FClothingAssetMeshBinding>& ClothingBindings)
 {
 	for (int32 LODIndex = 0; LODIndex < SkeletalMesh->GetImportedModel()->LODModels.Num(); ++LODIndex)
diff --git a/Engine/Source/Developer/SkeletalMeshUtilitiesCommon/Public/LODUtilities.h b/Engine/Source/Developer/SkeletalMeshUtilitiesCommon/Public/LODUtilities.h
--- a/Engine/Source/Developer/SkeletalMeshUtilitiesCommon/Public/LODUtilities.h
+++ b/Engine/Source/Developer/SkeletalMeshUtilitiesCommon/Public/LODUtilities.h
@@ -132,6 +132,21 @@ public:
 	 */
 	static void BuildMorphTargets(USkeletalMesh* SkeletalMesh, class FSkeletalMeshImportData &ImportData, int32 LODIndex, bool ShouldImportNormals, bool ShouldImportTangents, bool bUseMikkTSpace, const FOverlappingThresholds& Thresholds);
 
+	/*
+	 * Same as above but use normals from the source mesh description to build up the morph targets. 
+	 */
+	static void BuildMorphTargets(
+		USkeletalMesh* SkeletalMesh,
+		const FMeshDescription& SkeletalMeshModel,
+		FSkeletalMeshImportData& ImportData,
+		int32 LODIndex,
+		bool ShouldImportNormals,
+		bool ShouldImportTangents,
+		bool bUseMikkTSpace,
+		const FOverlappingThresholds& Thresholds
+		);
+	
+	
 	/**
 	 *	This function apply the skinning weights from asource skeletal mesh to the destination skeletal mesh.
 	 *  The Destination will receive the weights has the alternate weights.
diff --git a/Engine/Source/Runtime/Core/Private/UObject/UE5ReleaseStreamObjectVersion.cpp b/Engine/Source/Runtime/Core/Private/UObject/UE5ReleaseStreamObjectVersion.cpp
--- a/Engine/Source/Runtime/Core/Private/UObject/UE5ReleaseStreamObjectVersion.cpp
+++ b/Engine/Source/Runtime/Core/Private/UObject/UE5ReleaseStreamObjectVersion.cpp
@@ -13,7 +13,7 @@ TMap<FGuid, FGuid> FUE5ReleaseStreamObjectVersion::GetSystemGuids()
 	SystemGuids.Add(DevGuids.NANITE_DERIVEDDATA_VER, FGuid("AC88CFBDDA614A9C8488F64D2DAF699D"));
 	SystemGuids.Add(DevGuids.NIAGARASHADERMAP_DERIVEDDATA_VER, FGuid("6BE35B6FACB34568970120F9BC8DAB80"));
 	SystemGuids.Add(DevGuids.Niagara_LatestScriptCompileVersion, FGuid("43A777C54EB24E9894525B8D04529F23"));
-	SystemGuids.Add(DevGuids.SkeletalMeshDerivedDataVersion, FGuid("E55E25709ADD40B385EFDB6E3BB393A7"));
+	SystemGuids.Add(DevGuids.SkeletalMeshDerivedDataVersion, FGuid("3DEA297061DA4BB991A649D0B2F08F18"));
 	SystemGuids.Add(DevGuids.STATICMESH_DERIVEDDATA_VER, FGuid("7A556175518D458AA8786ED69AB8DDE7"));
 
 	return SystemGuids;
diff --git a/Engine/Source/Runtime/Engine/Classes/Engine/SkeletalMeshSourceModel.h b/Engine/Source/Runtime/Engine/Classes/Engine/SkeletalMeshSourceModel.h
--- a/Engine/Source/Runtime/Engine/Classes/Engine/SkeletalMeshSourceModel.h
+++ b/Engine/Source/Runtime/Engine/Classes/Engine/SkeletalMeshSourceModel.h
@@ -138,6 +138,12 @@ private:
 	 */
 	bool LoadMeshDescriptionFromBulkData(FMeshDescription& OutMeshDescription) const;
 
+	/**
+	 * If necessary, convert morph targets from the old FVector3f[2] storage to the split representation
+	 * where we store position per vertex, and normal per vertex instance
+	 */
+	static void UpgradeMorphTargets(FMeshDescription& InOutMeshDescription);
+
 	/**
 	 * Convert already stored old raw import format to mesh description and commit to bulk data.
 	 * The new bulk data will default to using the mesh hash as the GUID, so that any future
diff --git a/Engine/Source/Runtime/Engine/Private/Rendering/SkeletalMeshLODImporterData.cpp b/Engine/Source/Runtime/Engine/Private/Rendering/SkeletalMeshLODImporterData.cpp
--- a/Engine/Source/Runtime/Engine/Private/Rendering/SkeletalMeshLODImporterData.cpp
+++ b/Engine/Source/Runtime/Engine/Private/Rendering/SkeletalMeshLODImporterData.cpp
@@ -1996,8 +1996,10 @@ bool FSkeletalMeshImportData::GetMeshDescription(const USkeletalMesh* InSkeletal
 				UE_ASSET_LOG(LogSkeletalMeshLODImporterData, Warning, InSkeletalMesh, TEXT("Duplicate morph target '%s' found, renamed to '%s'."),
 					*MorphTargetNameBase, *MorphTargetNameUsed);  
 			}
-			
-			if (ensure(MeshAttributes.RegisterMorphTargetAttribute(*MorphTargetNameUsed)))
+
+			// Register a morph attribute but not with normals since the import data doesn't
+			// have any.
+			if (ensure(MeshAttributes.RegisterMorphTargetAttribute(*MorphTargetNameUsed, false)))
 			{
 				ValidMorphTargets.Add(MorphTargetIndex, MorphTargetNameUsed);
 				CheckDuplicateSet.Add(MorphTargetNameUsed);
@@ -2305,7 +2307,7 @@ bool FSkeletalMeshImportData::GetMeshDescription(const USkeletalMesh* InSkeletal
 		const TArray<FVector3f>& MorphPoints = MorphTargets[MorphTargetIndex].Points;
 		const TSet<uint32>& ModifiedPoints = MorphTargetModifiedPoints[MorphTargetIndex];
 
-		FMorphTargetVertexAttributesRef MorphTargetRef = MeshAttributes.GetVertexMorphTarget(*MorphTargetName);
+		TVertexAttributesRef<FVector3f> PositionDelta = MeshAttributes.GetVertexMorphPositionDelta(*MorphTargetName);
 
 		int32 ModifiedPointIndex = 0;
 		// This relies on the dubious assumptions that TSet values are retained in insertion order.
@@ -2313,7 +2315,7 @@ bool FSkeletalMeshImportData::GetMeshDescription(const USkeletalMesh* InSkeletal
 		{
 			const FVertexID VertexID = VertexIDMap[PointIndex];
 			
-			MorphTargetRef.SetPositionDelta(VertexID, MorphPoints[ModifiedPointIndex] - Points[PointIndex]);
+			PositionDelta.Set(VertexID, MorphPoints[ModifiedPointIndex] - Points[PointIndex]);
 			ModifiedPointIndex++;
 		}
 	}
@@ -2692,13 +2694,12 @@ FSkeletalMeshImportData FSkeletalMeshImportData::CreateFromMeshDescription(const
 	//Copy the vertex positions and the influences
 
 	//Reserve the point and influences
-	SkelMeshImportData.Points.SetNumUninitialized(InMeshDescription.Vertices().Num());
-	SkelMeshImportData.Influences.Reserve(InMeshDescription.Vertices().Num() * 4);
-	SkelMeshImportData.PointToRawMap.SetNumUninitialized(InMeshDescription.Vertices().Num());
+	SkelMeshImportData.Points.SetNumUninitialized(InMeshDescription.Vertices().GetArraySize());
+	SkelMeshImportData.Influences.Reserve(InMeshDescription.Vertices().GetArraySize() * 4);
+	SkelMeshImportData.PointToRawMap.SetNumUninitialized(InMeshDescription.Vertices().GetArraySize());
 
 	for (FVertexID VertexID : InMeshDescription.Vertices().GetElementIDs())
 	{
-		// We can use GetValue here because the mesh description was compacted before the copy
 		SkelMeshImportData.Points[VertexID.GetValue()] = VertexPositions[VertexID];
 		if (VertexSkinWeights.IsValid())
 		{
@@ -2744,14 +2745,35 @@ FSkeletalMeshImportData FSkeletalMeshImportData::CreateFromMeshDescription(const
 			break;
 		}
 	}
-	
+
 	//////////////////////////////////////////////////////////////////////////
-	//Copy the triangle and vertex instances
-	SkelMeshImportData.Faces.AddZeroed(InMeshDescription.Triangles().Num());
-	SkelMeshImportData.Wedges.Reserve(InMeshDescription.VertexInstances().Num());
+	// Copy vertex instances into wedges.
 	SkelMeshImportData.NumTexCoords = FMath::Min<int32>(VertexInstanceUVs.GetNumChannels(), (int32)MAX_TEXCOORDS);
 	SkelMeshImportData.bHasNormals = bHaveValidNormals;
 	SkelMeshImportData.bHasTangents = bHaveValidTangents;
+	
+	SkelMeshImportData.Wedges.SetNumZeroed(InMeshDescription.VertexInstances().GetArraySize());
+	for (FVertexInstanceID VertexInstanceID: InMeshDescription.VertexInstances().GetElementIDs())
+	{
+		SkeletalMeshImportData::FVertex& Wedge = SkelMeshImportData.Wedges[VertexInstanceID.GetValue()];
+		
+		Wedge.VertexIndex = static_cast<uint32>(InMeshDescription.GetVertexInstanceVertex(VertexInstanceID).GetValue());
+		Wedge.MatIndex = 0;			// We set this later -- not that this is actually used by any internal process.
+		constexpr bool bSRGB = false; //avoid linear to srgb conversion
+		Wedge.Color = FLinearColor(VertexInstanceColors[VertexInstanceID]).ToFColor(bSRGB);
+		if (Wedge.Color != FColor::White)
+		{
+			SkelMeshImportData.bHasVertexColors = true;
+		}
+		for (int32 UVChannelIndex = 0; UVChannelIndex < static_cast<int32>(SkelMeshImportData.NumTexCoords); ++UVChannelIndex)
+		{
+			Wedge.UVs[UVChannelIndex] = VertexInstanceUVs.Get(VertexInstanceID, UVChannelIndex);
+		}
+	}	
+	
+	//////////////////////////////////////////////////////////////////////////
+	// Copy the triangles
+	SkelMeshImportData.Faces.SetNumZeroed(InMeshDescription.Triangles().GetArraySize());
 	for (FTriangleID TriangleID : InMeshDescription.Triangles().GetElementIDs())
 	{
 		FPolygonGroupID PolygonGroupID = InMeshDescription.GetTrianglePolygonGroup(TriangleID);
@@ -2765,24 +2787,10 @@ FSkeletalMeshImportData FSkeletalMeshImportData::CreateFromMeshDescription(const
 		{
 			Face.SmoothingGroups = FaceSmoothingMasks[TriangleIndex];
 		}
-		//Create the wedges
 		for (int32 Corner = 0; Corner < 3; ++Corner)
 		{
 			FVertexInstanceID VertexInstanceID = VertexInstances[Corner];
-			SkeletalMeshImportData::FVertex Wedge;
-			Wedge.VertexIndex = (uint32)InMeshDescription.GetVertexInstanceVertex(VertexInstances[Corner]).GetValue();
-			Wedge.MatIndex = Face.MatIndex;
-			const bool bSRGB = false; //avoid linear to srgb conversion
-			Wedge.Color = FLinearColor(VertexInstanceColors[VertexInstanceID]).ToFColor(bSRGB);
-			if (Wedge.Color != FColor::White)
-			{
-				SkelMeshImportData.bHasVertexColors = true;
-			}
-			for (int32 UVChannelIndex = 0; UVChannelIndex < (int32)(SkelMeshImportData.NumTexCoords); ++UVChannelIndex)
-			{
-				Wedge.UVs[UVChannelIndex] = VertexInstanceUVs.Get(VertexInstanceID, UVChannelIndex);
-			}
-
+			
 			if (bHaveValidNormals)
 			{
 				Face.TangentZ[Corner] = VertexInstanceNormals[VertexInstanceID];
@@ -2802,7 +2810,9 @@ FSkeletalMeshImportData FSkeletalMeshImportData::CreateFromMeshDescription(const
 				Face.TangentX[Corner] = FVector3f::ZeroVector;
 			}
 
-			Face.WedgeIndex[Corner] = SkelMeshImportData.Wedges.Add(Wedge);
+			const int32 WedgeIndex = VertexInstanceID.GetValue();
+			Face.WedgeIndex[Corner] = WedgeIndex;
+			SkelMeshImportData.Wedges[WedgeIndex].MatIndex = Face.MatIndex;
 		}
 	}
 
@@ -2844,7 +2854,7 @@ FSkeletalMeshImportData FSkeletalMeshImportData::CreateFromMeshDescription(const
 	// Copy morph targets.
 	for (const FName MorphTargetName: MeshAttributes.GetMorphTargetNames())
 	{
-		FMorphTargetVertexAttributesConstRef MorphTargetAttribute = MeshAttributes.GetVertexMorphTarget(MorphTargetName);
+		TVertexAttributesConstRef<FVector3f> MorphTargetPosDeltaAttribute = MeshAttributes.GetVertexMorphPositionDelta(MorphTargetName);
 
 		SkelMeshImportData.MorphTargetNames.Add(MorphTargetName.ToString());
 		FSkeletalMeshImportData& MorphTarget = SkelMeshImportData.MorphTargets.AddDefaulted_GetRef();
@@ -2852,7 +2862,7 @@ FSkeletalMeshImportData FSkeletalMeshImportData::CreateFromMeshDescription(const
 
 		for (FVertexID VertexID : InMeshDescription.Vertices().GetElementIDs())
 		{
-			FVector3f PositionDelta = MorphTargetAttribute.GetPositionDelta(VertexID);
+			FVector3f PositionDelta = MorphTargetPosDeltaAttribute.Get(VertexID);
 			if (!PositionDelta.IsNearlyZero())
 			{
 				MorphTarget.Points.Add(PositionDelta + SkelMeshImportData.Points[VertexID.GetValue()]);
diff --git a/Engine/Source/Runtime/Engine/Private/SkeletalMesh.cpp b/Engine/Source/Runtime/Engine/Private/SkeletalMesh.cpp
--- a/Engine/Source/Runtime/Engine/Private/SkeletalMesh.cpp
+++ b/Engine/Source/Runtime/Engine/Private/SkeletalMesh.cpp
@@ -3334,6 +3334,10 @@ void USkeletalMesh::BeginPostLoadInternal(FSkinnedAssetPostLoadContext& Context)
 				LODModel.GetMeshDescription(this, LODIndex, MeshDescription);
 				CreateMeshDescription(LODIndex, MoveTemp(MeshDescription));
 				CommitMeshDescription(LODIndex);
+				
+				// Ensure normals aren't automatically computed when we rebuild.
+				FSkeletalMeshBuildSettings& BuildSettings = MeshLODInfo->BuildSettings;
+				BuildSettings.bRecomputeNormals = false;
 
 				// Reset the reduction settings so that we don't re-reduce the mesh and possibly lose morph targets
 				// in the process.
@@ -4084,7 +4088,7 @@ bool USkeletalMesh::RenameMorphTarget(FName InOldName, FName InNewName)
 
 	FScopedTransaction Transaction(LOCTEXT("RenameMorphTarget", "Rename Morph Target"));
 
-	// Unregister the morph target (but dont invalidate renderdata yet, we will recreate it below in RegisterMorphTarget)
+	// Unregister the morph target (but dont invalidate render data yet, we will recreate it below in RegisterMorphTarget)
 	UnregisterMorphTarget(MorphTarget, false);
 
 	Modify();
@@ -4099,12 +4103,23 @@ bool USkeletalMesh::RenameMorphTarget(FName InOldName, FName InNewName)
 		if (MeshAttributes.GetMorphTargetNames().Contains(InOldName))
 		{
 			ModifyMeshDescription(LODIndex);
-			if (MeshAttributes.RegisterMorphTargetAttribute(InNewName))
+
+			const bool bNeedNormals = MeshAttributes.GetVertexInstanceMorphNormalDelta(InOldName).IsValid();
+			
+			if (MeshAttributes.RegisterMorphTargetAttribute(InNewName, bNeedNormals))
 			{
-				const FMorphTargetVertexAttributesConstRef SourceMorphRef{ MeshAttributes.GetVertexMorphTarget(InOldName) };
-				FMorphTargetVertexAttributesRef TargetMorphRef{ MeshAttributes.GetVertexMorphTarget(InNewName) };
+				const TVertexAttributesConstRef<FVector3f> SourcePositionDelta{ MeshAttributes.GetVertexMorphPositionDelta(InOldName) };
+				TVertexAttributesRef<FVector3f> TargetPositionDelta{ MeshAttributes.GetVertexMorphPositionDelta(InNewName) };
+
+				TargetPositionDelta.Copy(SourcePositionDelta);
 
-				TargetMorphRef.Copy(SourceMorphRef);
+				if(bNeedNormals)
+				{
+					const TVertexInstanceAttributesConstRef<FVector3f> SourceNormalDelta{ MeshAttributes.GetVertexInstanceMorphNormalDelta(InOldName) };
+					TVertexInstanceAttributesRef<FVector3f> TargetNormalDelta{ MeshAttributes.GetVertexInstanceMorphNormalDelta(InNewName) };
+					TargetNormalDelta.Copy(SourceNormalDelta);
+				}
+				
 				MeshAttributes.UnregisterMorphTargetAttribute(InOldName);
 			}
 		}
diff --git a/Engine/Source/Runtime/Engine/Private/SkeletalMeshLODModel.cpp b/Engine/Source/Runtime/Engine/Private/SkeletalMeshLODModel.cpp
--- a/Engine/Source/Runtime/Engine/Private/SkeletalMeshLODModel.cpp
+++ b/Engine/Source/Runtime/Engine/Private/SkeletalMeshLODModel.cpp
@@ -1363,6 +1363,10 @@ void FSkeletalMeshLODModel::GetMeshDescription(const USkeletalMesh *InSkeletalMe
 	FSkeletalMeshAttributes::FBoneParentIndexAttributesRef BoneParentIndices = MeshAttributes.GetBoneParentIndices();
 	FSkeletalMeshAttributes::FBonePoseAttributesRef BonePoses = MeshAttributes.GetBonePoses();
 
+	// If the RawPointIndices2 is a map from the IndexBuffer to the original import vertices,
+	// this most likely came from a USD/Alembic import where the normals are set explicitly.
+	const bool bMorphTargetIncludeNormals = (RawPointIndices2.Num() == IndexBuffer.Num() && MeshToImportVertexMap.IsEmpty());
+	
 	TArray<TPair<FName, const FMorphTargetLODModel*>> MorphTargets;
 	for (UMorphTarget* MorphTargetSource: InSkeletalMesh->GetMorphTargets())
 	{
@@ -1378,7 +1382,7 @@ void FSkeletalMeshLODModel::GetMeshDescription(const USkeletalMesh *InSkeletalMe
 		}
 		
 		MorphTargets.Emplace(Name, &MorphTargetSource->GetMorphLODModels()[InLODIndex]);
-		MeshAttributes.RegisterMorphTargetAttribute(Name);
+		MeshAttributes.RegisterMorphTargetAttribute(Name, bMorphTargetIncludeNormals);
 	}
 
 	for (const TPair<FName, FImportedSkinWeightProfileData>& SkinWeightProfileInfo: SkinWeightProfiles)
@@ -1401,22 +1405,20 @@ void FSkeletalMeshLODModel::GetMeshDescription(const USkeletalMesh *InSkeletalMe
 	// Map the section vertices back to the import vertices to remove seams, but only if there's
 	// mapping available.
 	TArray<int32> SourceToTargetVertexMap; 
-	TSet<int32> ProcessedTargetVertex;
 
 	int32 TargetVertexCount = 0;
+	
 	if (RawPointIndices2.Num() == NumVertices)
 	{
 		SourceToTargetVertexMap.Reserve(RawPointIndices2.Num());
-		for (const uint32 Index: RawPointIndices2)
+		
+		for (const uint32 VertexIndex: RawPointIndices2)
 		{
-			SourceToTargetVertexMap.Add(Index);
-			ProcessedTargetVertex.Add(Index);
-
-			TargetVertexCount = FMath::Max(TargetVertexCount, static_cast<int32>(Index));
+			SourceToTargetVertexMap.Add(VertexIndex);
+			TargetVertexCount = FMath::Max(TargetVertexCount, static_cast<int32>(VertexIndex));
 		}
 
 		TargetVertexCount += 1;
-		ProcessedTargetVertex.Reset();
 	}
 	else
 	{
@@ -1435,12 +1437,22 @@ void FSkeletalMeshLODModel::GetMeshDescription(const USkeletalMesh *InSkeletalMe
 		VertexIDs.Add(OutMeshDescription.CreateVertex());
 	}
 
+	// Mapping to go from morph target vertices to vertex instances.
+	TMultiMap<int32, FVertexInstanceID> SourceVertexToVertexInstanceMap;
+	if (bMorphTargetIncludeNormals)
+	{
+		SourceVertexToVertexInstanceMap.Reserve(IndexBuffer.Num());
+	}
+
 	// Ensure we have enough channels to store all the defined UV coordinates.
 	VertexInstanceUVs.SetNumChannels(static_cast<int32>(NumTexCoords));
 	
 	const TArray<FSkeletalMaterial>& Materials = InSkeletalMesh->GetMaterials();
 	const bool bHasVertexColors = EnumHasAllFlags(InSkeletalMesh->GetVertexBufferFlags(), ESkeletalMeshVertexFlags::HasVertexColors);
 
+	TSet<int32> ProcessedTargetVertex;
+	ProcessedTargetVertex.Reserve(TargetVertexCount);
+	
 	// Convert sections to polygon groups, each with their own material.
 	for (int32 SectionIndex = 0; SectionIndex < Sections.Num(); SectionIndex++)
 	{
@@ -1457,6 +1469,7 @@ void FSkeletalMeshLODModel::GetMeshDescription(const USkeletalMesh *InSkeletalMe
 			{
 				continue;
 			}
+			ProcessedTargetVertex.Add(TargetVertexIndex);
 			
 			const FVertexID VertexID = VertexIDs[TargetVertexIndex];
 
@@ -1503,6 +1516,11 @@ void FSkeletalMeshLODModel::GetMeshDescription(const USkeletalMesh *InSkeletalMe
 				const FVertexID VertexID = VertexIDs[TargetVertexIndex];
 				const FVertexInstanceID VertexInstanceID = OutMeshDescription.CreateVertexInstance(VertexID);
 
+				if (bMorphTargetIncludeNormals)
+				{
+					SourceVertexToVertexInstanceMap.Add(SourceVertexIndex, VertexInstanceID);
+				}
+
 				const FSoftSkinVertex& SourceVertex = SourceVertices[SourceVertexIndex - Section.BaseVertexIndex];
 
 				VertexInstanceNormals.Set(VertexInstanceID, SourceVertex.TangentZ);
@@ -1533,14 +1551,27 @@ void FSkeletalMeshLODModel::GetMeshDescription(const USkeletalMesh *InSkeletalMe
 	// Copy morph targets.
 	for (TPair<FName, const FMorphTargetLODModel*>& MorphSource: MorphTargets)
 	{
-		FMorphTargetVertexAttributesRef MorphTarget = MeshAttributes.GetVertexMorphTarget(MorphSource.Key);
+		TVertexAttributesRef<FVector3f> PositionDelta = MeshAttributes.GetVertexMorphPositionDelta(MorphSource.Key);
+		TVertexInstanceAttributesRef<FVector3f> NormalDelta = MeshAttributes.GetVertexInstanceMorphNormalDelta(MorphSource.Key);
+		TArray<FVertexInstanceID> VertexInstanceIDs;
 
 		for (const FMorphTargetDelta& Delta: MorphSource.Value->Vertices)
 		{
 			const int32 TargetVertexIndex = SourceToTargetVertexMap[Delta.SourceIdx];
 			const FVertexID VertexID = VertexIDs[TargetVertexIndex];
-			
-			MorphTarget.SetPositionAndTangentZDelta(VertexID, Delta.PositionDelta, Delta.TangentZDelta);
+
+			PositionDelta.Set(VertexID, Delta.PositionDelta);
+
+			if (bMorphTargetIncludeNormals)
+			{
+				VertexInstanceIDs.Reset();
+				SourceVertexToVertexInstanceMap.MultiFind(Delta.SourceIdx, VertexInstanceIDs);
+
+				for (FVertexInstanceID VertexInstanceID: VertexInstanceIDs)
+				{
+					NormalDelta.Set(VertexInstanceID, Delta.TangentZDelta); 
+				}
+			}
 		}
 	}
 
diff --git a/Engine/Source/Runtime/Engine/Private/SkeletalMeshSourceModel.cpp b/Engine/Source/Runtime/Engine/Private/SkeletalMeshSourceModel.cpp
--- a/Engine/Source/Runtime/Engine/Private/SkeletalMeshSourceModel.cpp
+++ b/Engine/Source/Runtime/Engine/Private/SkeletalMeshSourceModel.cpp
@@ -234,7 +234,9 @@ void FSkeletalMeshSourceModel::EnsureRawMeshBulkDataIsConvertedToNew()
 }
 
 
-bool FSkeletalMeshSourceModel::LoadMeshDescriptionFromBulkData(FMeshDescription& OutMeshDescription) const
+bool FSkeletalMeshSourceModel::LoadMeshDescriptionFromBulkData(
+	FMeshDescription& OutMeshDescription
+	) const
 {
 	TRACE_CPUPROFILER_EVENT_SCOPE(FSkeletalMeshSourceData::LoadMeshDescriptionFromBulkData);
 	
@@ -255,10 +257,50 @@ bool FSkeletalMeshSourceModel::LoadMeshDescriptionFromBulkData(FMeshDescription&
 	}
 	
 	MeshDescriptionBulkData->GetBulkData().LoadMeshDescription(OutMeshDescription);
+	
+	// If this mesh is stored with the older representation, update it now.
+	UpgradeMorphTargets(OutMeshDescription);
+	
 	return true;
 }
 
+void FSkeletalMeshSourceModel::UpgradeMorphTargets(
+	FMeshDescription& InOutMeshDescription
+	)	
+{
+	TArray<FName> AttributeNames;
+	InOutMeshDescription.VertexAttributes().GetAttributeNames(AttributeNames);
 
+	TArray<FVector3f> SavedPointDelta;
+	for (const FName AttributeName: AttributeNames)
+	{
+		if (!FSkeletalMeshAttributes::IsMorphTargetAttribute(AttributeName))
+		{
+			continue;
+		}
+		
+		TVertexAttributesConstRef<TArrayView<FVector3f>> OldMorphTarget = InOutMeshDescription.VertexAttributes().GetAttributesRef<TArrayView<FVector3f>>(AttributeName);
+
+		if (!OldMorphTarget.IsValid())
+		{
+			continue;
+		}
+		// Grab all the points and store them away. We then re-register the attribute(s) with the correct details. We only
+		// need the point deltas, since that's all that was stored before the storage change.
+		SavedPointDelta.SetNumUninitialized(InOutMeshDescription.Vertices().GetArraySize());
+
+		for (FVertexID VertexID: InOutMeshDescription.Vertices().GetElementIDs())
+		{
+			SavedPointDelta[VertexID.GetValue()] = OldMorphTarget.Get(VertexID)[0];
+		}
+
+		TVertexAttributesRef<FVector3f> PointDelta = InOutMeshDescription.VertexAttributes().RegisterAttribute<FVector3f>(AttributeName, 1, FVector3f::ZeroVector, EMeshAttributeFlags::None);
+		for (FVertexID VertexID: InOutMeshDescription.Vertices().GetElementIDs())
+		{
+			PointDelta.Set(VertexID, SavedPointDelta[VertexID.GetValue()]);
+		}
+	}
+}
 
 
 void FSkeletalMeshSourceModel::ConvertRawMeshToMeshDescriptionBulkData()
diff --git a/Engine/Source/Runtime/MeshDescription/Public/MeshAttributeArray.h b/Engine/Source/Runtime/MeshDescription/Public/MeshAttributeArray.h
--- a/Engine/Source/Runtime/MeshDescription/Public/MeshAttributeArray.h
+++ b/Engine/Source/Runtime/MeshDescription/Public/MeshAttributeArray.h
@@ -1007,7 +1007,7 @@ template <typename ElementIDType, typename AttributeType>
 void TMeshAttributesRef<ElementIDType, AttributeType>::Copy(TMeshAttributesRef<ElementIDType, const AttributeType> Src, const int32 DestChannel, const int32 SrcChannel)
 {
 	check(Src.IsValid());
-	const TMeshAttributeArrayBase<AttributeType>& SrcArray = static_cast<const ArrayType*>(Src->ArrayPtr)->GetArrayForChannel(SrcChannel);
+	const TMeshAttributeArrayBase<AttributeType>& SrcArray = static_cast<const ArrayType*>(Src.ArrayPtr)->GetArrayForChannel(SrcChannel);
 	TMeshAttributeArrayBase<AttributeType>& DestArray = static_cast<ArrayType*>(ArrayPtr)->GetArrayForChannel(DestChannel);
 	const int32 Num = FMath::Min(SrcArray.Num(), DestArray.Num());
 	for (int32 Index = 0; Index < Num; Index++)
diff --git a/Engine/Source/Runtime/SkeletalMeshDescription/Private/SkeletalMeshAttributes.cpp b/Engine/Source/Runtime/SkeletalMeshDescription/Private/SkeletalMeshAttributes.cpp
--- a/Engine/Source/Runtime/SkeletalMeshDescription/Private/SkeletalMeshAttributes.cpp
+++ b/Engine/Source/Runtime/SkeletalMeshDescription/Private/SkeletalMeshAttributes.cpp
@@ -141,7 +141,10 @@ FSkinWeightsVertexAttributesRef FSkeletalMeshAttributes::GetVertexSkinWeightsFro
 }
 
 
-bool FSkeletalMeshAttributes::RegisterMorphTargetAttribute(const FName InMorphTargetName)
+bool FSkeletalMeshAttributes::RegisterMorphTargetAttribute(
+	const FName InMorphTargetName,
+	const bool bIncludeNormals
+	)
 {
 	if (InMorphTargetName.IsNone())
 	{
@@ -154,15 +157,21 @@ bool FSkeletalMeshAttributes::RegisterMorphTargetAttribute(const FName InMorphTa
 		return false;
 	}
 
-	// Already has this attribute?
-	if (MeshDescription.VertexAttributes().HasAttribute(AttributeName))
+	bool bSuccess = MeshDescription.VertexAttributes().RegisterAttribute<FVector3f>(AttributeName, 1, FVector3f::ZeroVector, EMeshAttributeFlags::None).IsValid();
+	if (bSuccess && bIncludeNormals)
 	{
-		return false;
+		// Register normal attribute, if requested, if it fails, then we have to unregister the vertex attribute too.
+		bSuccess = MeshDescription.VertexInstanceAttributes().RegisterAttribute<FVector3f>(AttributeName, 1, FVector3f::ZeroVector, EMeshAttributeFlags::None).IsValid();
+		if (!bSuccess)
+		{
+			MeshDescription.VertexAttributes().UnregisterAttribute(AttributeName);
+		}
 	}
 
-	return MeshDescription.VertexAttributes().RegisterAttribute<FVector3f[2]>(AttributeName, 1, FVector3f::ZeroVector, EMeshAttributeFlags::None).IsValid();
+	return bSuccess;
 }
 
+
 bool FSkeletalMeshAttributes::UnregisterMorphTargetAttribute(const FName InMorphTargetName)
 {
 	if (InMorphTargetName.IsNone())
@@ -186,18 +195,16 @@ bool FSkeletalMeshAttributes::UnregisterMorphTargetAttribute(const FName InMorph
 	return true;
 }
 
-FMorphTargetVertexAttributesRef FSkeletalMeshAttributes::GetVertexMorphTarget(const FName InMorphTargetName)
+
+TVertexAttributesRef<FVector3f> FSkeletalMeshAttributes::GetVertexMorphPositionDelta(const FName InMorphTargetName)
 {
-	return MeshDescription.VertexAttributes().GetAttributesRef<TArrayView<FVector3f>>(CreateMorphTargetAttributeName(InMorphTargetName));
+	return MeshDescription.VertexAttributes().GetAttributesRef<FVector3f>(CreateMorphTargetAttributeName(InMorphTargetName));
 }
 
-FMorphTargetVertexAttributesRef FSkeletalMeshAttributes::GetVertexMorphTargetFromAttributeName(const FName InAttributeName)
+
+TVertexInstanceAttributesRef<FVector3f> FSkeletalMeshAttributes::GetVertexInstanceMorphNormalDelta(const FName InMorphTargetName)
 {
-	if (!IsMorphTargetAttribute(InAttributeName))
-	{
-		return {};
-	}
-	return MeshDescription.VertexAttributes().GetAttributesRef<TArrayView<FVector3f>>(InAttributeName);
+	return MeshDescription.VertexInstanceAttributes().GetAttributesRef<FVector3f>(CreateMorphTargetAttributeName(InMorphTargetName));
 }
 
 
@@ -461,18 +468,14 @@ FName FSkeletalMeshAttributesShared::GetMorphTargetNameFromAttribute(const FName
 	return NAME_None;
 }
 
-FMorphTargetVertexAttributesConstRef FSkeletalMeshAttributesShared::GetVertexMorphTarget(const FName InMorphTargetName) const
+TVertexAttributesConstRef<FVector3f> FSkeletalMeshAttributesShared::GetVertexMorphPositionDelta(const FName InMorphTargetName) const
 {
-	return MeshDescriptionShared.VertexAttributes().GetAttributesRef<TArrayView<FVector3f>>(CreateMorphTargetAttributeName(InMorphTargetName));
+	return MeshDescriptionShared.VertexAttributes().GetAttributesRef<FVector3f>(CreateMorphTargetAttributeName(InMorphTargetName));
 }
 
-FMorphTargetVertexAttributesConstRef FSkeletalMeshAttributesShared::GetVertexMorphTargetFromAttributeName(const FName InAttributeName) const
+TVertexInstanceAttributesConstRef<FVector3f> FSkeletalMeshAttributesShared::GetVertexInstanceMorphNormalDelta(const FName InMorphTargetName) const
 {
-	if (!IsMorphTargetAttribute(InAttributeName))
-	{
-		return {};
-	}
-	return MeshDescriptionShared.VertexAttributes().GetAttributesRef<TArrayView<FVector3f>>(InAttributeName);
+	return MeshDescriptionShared.VertexInstanceAttributes().GetAttributesRef<FVector3f>(CreateMorphTargetAttributeName(InMorphTargetName));
 }
 
 
diff --git a/Engine/Source/Runtime/SkeletalMeshDescription/Public/MorphTargetAttributesRef.h b/Engine/Source/Runtime/SkeletalMeshDescription/Public/MorphTargetAttributesRef.h
deleted file mode 100644
diff --git a/Engine/Source/Runtime/SkeletalMeshDescription/Public/SkeletalMeshAttributes.h b/Engine/Source/Runtime/SkeletalMeshDescription/Public/SkeletalMeshAttributes.h
--- a/Engine/Source/Runtime/SkeletalMeshDescription/Public/SkeletalMeshAttributes.h
+++ b/Engine/Source/Runtime/SkeletalMeshDescription/Public/SkeletalMeshAttributes.h
@@ -7,7 +7,6 @@
 #include "Containers/Array.h"
 #include "HAL/Platform.h"
 #include "MeshAttributeArray.h"
-#include "MorphTargetAttributesRef.h"
 #include "SkinWeightsAttributesRef.h"
 #include "StaticMeshAttributes.h"
 #include "UObject/NameTypes.h"
@@ -127,9 +126,15 @@ public:
 	/// Returns the name of a morph target given the attribute name. If the given attribute name is invalid, \c NAME_None is returned.
 	static SKELETALMESHDESCRIPTION_API FName GetMorphTargetNameFromAttribute(const FName InAttributeName);
 
-	SKELETALMESHDESCRIPTION_API FMorphTargetVertexAttributesConstRef GetVertexMorphTarget(const FName InMorphTargetName = NAME_None) const;
-	
-	SKELETALMESHDESCRIPTION_API FMorphTargetVertexAttributesConstRef GetVertexMorphTargetFromAttributeName(const FName InAttributeName = NAME_None) const;
+	/// Returns a read-only vertex position delta attribute for the given morph target. The delta is the difference between the point position of the
+	/// mesh and the point position of the morphed mesh. Adding this value to the point position of the mesh will return the morph point position.
+	SKELETALMESHDESCRIPTION_API TVertexAttributesConstRef<FVector3f> GetVertexMorphPositionDelta(const FName InMorphTargetName) const;
+
+	/// Returns a read-only vertex instance normal delta attribute for the given morph target. The delta is the difference between the normal vector 
+	/// of the mesh and the normal vector of the desired morphed mesh. Adding this value to the normal vector of the mesh will return the morph's 
+	/// normal vector for this vertex instance.
+	/// If the morph was not registered to include normals, this will return an invalid attribute. 
+	SKELETALMESHDESCRIPTION_API TVertexInstanceAttributesConstRef<FVector3f> GetVertexInstanceMorphNormalDelta(const FName InMorphTargetName) const;
 	
 	//
 	// Bones Methods
@@ -153,7 +158,7 @@ public:
 		return BoneElementsShared != nullptr;
 	}
 
-	/**  @return the number of the bones. 0 if bone element does not exist in the MeshDescripion. */
+	/**  @return the number of the bones. 0 if bone element does not exist in the FMeshDescription. */
 	SKELETALMESHDESCRIPTION_API int32 GetNumBones() const;
 
 	/** @return true, if the passed bone ID is valid */
@@ -261,21 +266,27 @@ public:
 	/// Register a new morph target with the given name. The attribute name will encode the user-defined morph target name and
 	/// it will be listed in GetMorphTargetNames(). Returns \c true if the morph target was successfully registered.
 	/// Returns \c false if the attribute was already registered or if \c InMorphTargetName is empty.
-	SKELETALMESHDESCRIPTION_API bool RegisterMorphTargetAttribute(const FName InMorphTargetName);
+	/// The position delta is stored as a vertex attributes, and the optional normal is stored as a vertex instance attribute to
+	/// allow for hard edges on morph targets.
+	/// \param InMorphTargetName The name of the morph target to add. Cannot be empty.
+	/// \param bIncludeNormals Set to \c true if per-vertex instance normals (Z-tangent) should be included. Otherwise the normals
+	///		will be automatically computed and only the position delta vertex attribute registered.
+	SKELETALMESHDESCRIPTION_API bool RegisterMorphTargetAttribute(const FName InMorphTargetName, const bool bIncludeNormals);
 
 	/// Unregister an existing morph target with the given name (as returned by GetMorphTargetNames()). Returns \c true if the morph target
 	/// was successfully unregistered.
 	/// Returns \c false if the attribute wasn't registered or if \c InMorphTargetName is empty.
 	SKELETALMESHDESCRIPTION_API bool UnregisterMorphTargetAttribute(const FName InMorphTargetName);
-	
-	/// Returns a specialized morph target vertex attribute, given a morph target name, to allow setting morph target data
-	/// in a simple fashion. If the morph target doesn't exist, the returned  attribute will be invalid and cannot be read from or written to. 
-	SKELETALMESHDESCRIPTION_API FMorphTargetVertexAttributesRef GetVertexMorphTarget(const FName InMorphTargetName = NAME_None);
-	
-	/// Returns a specialized morph target vertex attribute, given an attribute name, to allow setting morph target data
-	/// in a simple fashion. If the attribute doesn't exist, or the attribute does not refer to a morph target, the returned attribute
-	/// will be invalid and cannot be read from or written to. 
-	SKELETALMESHDESCRIPTION_API FMorphTargetVertexAttributesRef GetVertexMorphTargetFromAttributeName(const FName InAttributeName = NAME_None);
+
+	/// Returns a writable vertex position delta attribute for the given morph target. The delta is the difference between the point position of the
+	/// mesh and the point position of the morphed mesh. Adding this value to the point position of the mesh will return the morph point position.
+	SKELETALMESHDESCRIPTION_API TVertexAttributesRef<FVector3f> GetVertexMorphPositionDelta(const FName InMorphTargetName);
+
+	/// Returns a writable vertex instance normal delta attribute for the given morph target. The delta is the difference between the normal vector 
+	/// of the mesh and the normal vector of the desired morphed mesh. Adding this value to the normal vector of the mesh will return the morph's 
+	/// normal vector for this vertex instance.
+	/// If the morph was not registered to include normals, this will return an invalid attribute. 
+	SKELETALMESHDESCRIPTION_API TVertexInstanceAttributesRef<FVector3f> GetVertexInstanceMorphNormalDelta(const FName InMorphTargetName);
 	
 	//
 	// Bones Methods
