Interchange: Custom imported morph taget now are restore after a re-import.

#rnx

diff --git a/Engine/Plugins/Interchange/Runtime/Source/Import/Private/Mesh/InterchangeSkeletalMeshFactory.cpp b/Engine/Plugins/Interchange/Runtime/Source/Import/Private/Mesh/InterchangeSkeletalMeshFactory.cpp
--- a/Engine/Plugins/Interchange/Runtime/Source/Import/Private/Mesh/InterchangeSkeletalMeshFactory.cpp
+++ b/Engine/Plugins/Interchange/Runtime/Source/Import/Private/Mesh/InterchangeSkeletalMeshFactory.cpp
@@ -1286,6 +1286,10 @@ UInterchangeFactoryBase::FImportAssetResult UInterchangeSkeletalMeshFactory::Imp
 	//Make sure we can modify the skeletalmesh properties
 	FSkinnedAssetAsyncBuildScope AsyncBuildScope(SkeletalMesh);
 
+	TMap<FString, TArray<FLODUtilities::FMorphTargetLodBackupData>> BackupImportedMorphTargetData;
+	FLODUtilities::BackupCustomImportedMorphTargetData(SkeletalMesh, BackupImportedMorphTargetData);
+
+
 	FTransform GlobalOffsetTransform = FTransform::Identity;
 	bool bBakeMeshes = false;
 	if (UInterchangeCommonPipelineDataFactoryNode* CommonPipelineDataFactoryNode = UInterchangeCommonPipelineDataFactoryNode::GetUniqueInstance(Arguments.NodeContainer))
@@ -1637,10 +1641,18 @@ UInterchangeFactoryBase::FImportAssetResult UInterchangeSkeletalMeshFactory::Imp
 				}
 			}
 		}
-		//Store the original fbx import data the SkelMeshImportDataPtr should not be modified after this
-		PRAGMA_DISABLE_DEPRECATION_WARNINGS
-		SkeletalMesh->SaveLODImportedData(CurrentLodIndex, SkeletalMeshImportData);
-		PRAGMA_ENABLE_DEPRECATION_WARNINGS
+
+		//Store the original fbx import data the SkeletalMeshImportData should not be modified after this
+		{
+			FMeshDescription MeshDescription;
+			if (SkeletalMeshImportData.GetMeshDescription(nullptr, &SkeletalMesh->GetLODInfo(CurrentLodIndex)->BuildSettings, MeshDescription))
+			{
+				//Restore the morph target into the imported mesh description
+				FLODUtilities::RestoreCustomImportedMorphTargetData(SkeletalMesh, CurrentLodIndex, MeshDescription, BackupImportedMorphTargetData);
+				SkeletalMesh->CreateMeshDescription(CurrentLodIndex, MoveTemp(MeshDescription));
+				SkeletalMesh->CommitMeshDescription(CurrentLodIndex);
+			}
+		}
 
 		//Update the bounding box if we are importing the LOD 0
 		if(CurrentLodIndex == 0)
diff --git a/Engine/Plugins/Interchange/Runtime/Source/Pipelines/Private/InterchangeGenericAssetsPipeline.cpp b/Engine/Plugins/Interchange/Runtime/Source/Pipelines/Private/InterchangeGenericAssetsPipeline.cpp
--- a/Engine/Plugins/Interchange/Runtime/Source/Pipelines/Private/InterchangeGenericAssetsPipeline.cpp
+++ b/Engine/Plugins/Interchange/Runtime/Source/Pipelines/Private/InterchangeGenericAssetsPipeline.cpp
@@ -162,6 +162,11 @@ void UInterchangeGenericAssetsPipeline::AdjustSettingsForContext(EInterchangePip
 			}
 		}
 	}
+	else if (ImportType == EInterchangePipelineContext::AssetCustomMorphTargetImport
+		|| ImportType == EInterchangePipelineContext::AssetCustomMorphTargetReImport)
+	{
+		bUseSourceNameForAsset = false;
+	}
 
 	if (MaterialPipeline)
 	{
diff --git a/Engine/Source/Developer/SkeletalMeshUtilitiesCommon/Private/LODUtilities.cpp b/Engine/Source/Developer/SkeletalMeshUtilitiesCommon/Private/LODUtilities.cpp
--- a/Engine/Source/Developer/SkeletalMeshUtilitiesCommon/Private/LODUtilities.cpp
+++ b/Engine/Source/Developer/SkeletalMeshUtilitiesCommon/Private/LODUtilities.cpp
@@ -549,6 +549,9 @@ bool FLODUtilities::SetCustomLOD(USkeletalMesh* DestinationSkeletalMesh, USkelet
 	TArray<int32> ClothingAssetSectionIndices;
 	TArray<int32> ClothingAssetInternalLodIndices;
 
+	//SAve custom imported morph targets
+	TMap<FString, TArray<FMorphTargetLodBackupData>> BackupImportedMorphTargetData;
+
 	TArray<FSkinWeightProfileInfo> ExistingSkinWeightProfileInfos;
 	TArray<FSkeletalMeshImportData> ExistingAlternateImportDataPerLOD;
 
@@ -564,6 +567,9 @@ bool FLODUtilities::SetCustomLOD(USkeletalMesh* DestinationSkeletalMesh, USkelet
 	{
 		FLODUtilities::UnbindClothingAndBackup(DestinationSkeletalMesh, ClothingBindings, LodIndex);
 
+		//Backup the lod custom imported morph
+		BackupCustomImportedMorphTargetData(DestinationSkeletalMesh, BackupImportedMorphTargetData);
+
 		int32 ExistingLodCount = DestinationSkeletalMesh->GetLODNum();
 
 		//Extract all LOD Skin weight profiles data, we will re-apply them at the end
@@ -595,6 +601,20 @@ bool FLODUtilities::SetCustomLOD(USkeletalMesh* DestinationSkeletalMesh, USkelet
 		}
 	};
 
+	auto ReapplyCustomImportedMorphTarget = [&DestinationSkeletalMesh, &BackupImportedMorphTargetData, &LodIndex]()
+	{
+		if (FMeshDescription* MeshDescription = DestinationSkeletalMesh->GetMeshDescription(LodIndex))
+		{
+			if (FLODUtilities::RestoreCustomImportedMorphTargetData(DestinationSkeletalMesh, LodIndex, *MeshDescription, BackupImportedMorphTargetData))
+			{
+				USkeletalMesh::FCommitMeshDescriptionParams CommitParams;
+				CommitParams.bForceUpdate = false;
+				CommitParams.bMarkPackageDirty = false;
+				DestinationSkeletalMesh->CommitMeshDescription(LodIndex, CommitParams);
+			}
+		}
+	};
+
 	auto ReapplyAlternateSkinning = [&DestinationSkeletalMesh, &ExistingSkinWeightProfileInfos, &ExistingAlternateImportDataPerLOD, &LodIndex]()
 	{
 		if (ExistingSkinWeightProfileInfos.Num() > 0)
@@ -950,6 +970,8 @@ bool FLODUtilities::SetCustomLOD(USkeletalMesh* DestinationSkeletalMesh, USkelet
 
 	ReapplyClothing();
 
+	ReapplyCustomImportedMorphTarget();
+
 	ReapplyAlternateSkinning();
 	
 	// Notification of success
@@ -3991,6 +4013,132 @@ void FLODUtilities::RestoreClothingFromBackup(USkeletalMesh* SkeletalMesh, TArra
 	}
 }
 
+void FLODUtilities::BackupCustomImportedMorphTargetData(USkeletalMesh* SkeletalMesh, TMap<FString, TArray<FMorphTargetLodBackupData>>& BackupImportedMorphTargetData)
+{
+	const int32 LodCount = SkeletalMesh->GetLODNum();
+	constexpr int32 LodIndex0 = 0;
+	//Find all imported morph targets
+	if (const FSkeletalMeshLODInfo* LodInfo = SkeletalMesh->GetLODInfo(LodIndex0))
+	{
+		TSet<FString> ImportedMorphTargetNames;
+		LodInfo->ImportedMorphTargetSourceFilename.GetKeys(ImportedMorphTargetNames);
+		for (const FString& MorphTargetNameStr : ImportedMorphTargetNames)
+		{
+			FName MorphTargetName(*MorphTargetNameStr);
+			TArray<FMorphTargetLodBackupData>& MorphTargetLodDatas = BackupImportedMorphTargetData.FindOrAdd(MorphTargetNameStr);
+			for (int32 LodIndex = 0; LodIndex < LodCount; ++LodIndex)
+			{
+				FMorphTargetLodBackupData& MorphTargetLodData = MorphTargetLodDatas.AddDefaulted_GetRef();
+				if (SkeletalMesh->HasMeshDescription(LodIndex))
+				{
+					FMeshDescription& MeshDescription = *(SkeletalMesh->GetMeshDescription(LodIndex));
+					FSkeletalMeshConstAttributes SkeletalMeshAttributes(MeshDescription);
+
+					if (SkeletalMeshAttributes.HasMorphTargetPositionsAttribute(MorphTargetName))
+					{
+						MorphTargetLodData.bIsEmpty = false;
+						MorphTargetLodData.MorphPositionDeltas.Reserve(MeshDescription.Vertices().Num());
+						TVertexAttributesConstRef<FVector3f> VertexPositionDeltas = SkeletalMeshAttributes.GetVertexMorphPositionDelta(MorphTargetName);
+						for (FVertexID VertexID : MeshDescription.Vertices().GetElementIDs())
+						{
+							MorphTargetLodData.MorphPositionDeltas.Add(VertexPositionDeltas[VertexID]);
+						}
+						if (SkeletalMeshAttributes.HasMorphTargetNormalsAttribute(MorphTargetName))
+						{
+							MorphTargetLodData.MorphNormalDeltas.Reserve(MeshDescription.VertexInstances().Num());
+							TVertexInstanceAttributesConstRef<FVector3f> VertexInstanceNormalDeltas = SkeletalMeshAttributes.GetVertexInstanceMorphNormalDelta(MorphTargetName);
+							for (FVertexInstanceID VertexInstanceID : MeshDescription.VertexInstances().GetElementIDs())
+							{
+								MorphTargetLodData.MorphNormalDeltas.Add(VertexInstanceNormalDeltas[VertexInstanceID]);
+							}
+						}
+					}
+				}
+			}
+		}
+	}
+}
+
+bool FLODUtilities::RestoreCustomImportedMorphTargetData(USkeletalMesh* SkeletalMesh, const int32 LodIndex, FMeshDescription& LodMeshDescription, const TMap<FString, TArray<FMorphTargetLodBackupData>>& BackupImportedMorphTargetData)
+{
+	auto RemoveInvalidMorphTargetLOD = [SkeletalMesh, LodIndex](const FString& MorphTargetNameStr)
+		{
+			FName MorphTargetName(*MorphTargetNameStr);
+			if(UMorphTarget* MorphTarget = SkeletalMesh->FindMorphTarget(MorphTargetName))
+			{
+				MorphTarget->RemoveFromRoot();
+				MorphTarget->ClearFlags(RF_Standalone);
+				SkeletalMesh->UnregisterMorphTarget(MorphTarget, false);
+			}
+			if (FSkeletalMeshLODInfo* LodInfo = SkeletalMesh->GetLODInfo(LodIndex))
+			{
+				if (LodInfo->ImportedMorphTargetSourceFilename.Contains(MorphTargetNameStr))
+				{
+					LodInfo->ImportedMorphTargetSourceFilename.Remove(MorphTargetNameStr);
+				}
+			}
+		};
+
+	bool bIsMeshDescriptionModified = false;
+	for (const TPair<FString, TArray<FMorphTargetLodBackupData>>& ImportedMorphTargetPair : BackupImportedMorphTargetData)
+	{
+		const FString& MorphTargetNameStr = ImportedMorphTargetPair.Key;
+		FName MorphTargetName(*MorphTargetNameStr);
+
+		const TArray<FMorphTargetLodBackupData>& MorphTargetLodDatas = ImportedMorphTargetPair.Value;
+		if (!MorphTargetLodDatas.IsValidIndex(LodIndex))
+		{
+			continue;
+		}
+
+		const FMorphTargetLodBackupData& MorphTargetLodData = MorphTargetLodDatas[LodIndex];
+		if (MorphTargetLodData.bIsEmpty)
+		{
+			RemoveInvalidMorphTargetLOD(MorphTargetNameStr);
+			//Error out that the morph target is not valid anymore for this LOD
+			UE_ASSET_LOG(LogLODUtilities, Error, SkeletalMesh, TEXT("Cannot keep the generated morph target %s for LOD %d, because this LOD is now imported and the topology have change."), *MorphTargetNameStr, LodIndex);
+			continue;
+		}
+
+		FSkeletalMeshAttributes SkeletalMeshAttributes(LodMeshDescription);
+		if (SkeletalMeshAttributes.IsMorphTargetAttribute(MorphTargetName))
+		{
+			//The morph target already exist, this mean the re-import added this morph target
+			continue;
+		}
+
+		if (MorphTargetLodData.MorphPositionDeltas.Num() != LodMeshDescription.Vertices().Num())
+		{
+			RemoveInvalidMorphTargetLOD(MorphTargetNameStr);
+			//Error out that the morph target is not valid anymore for this LOD
+			UE_ASSET_LOG(LogLODUtilities, Error, SkeletalMesh, TEXT("Cannot keep custom morph target %s for LOD %d, because the re-import LOD now have a different mesh geometry topology."), *MorphTargetNameStr, LodIndex);
+			continue;
+		}
+
+		//We are good to go and we can add the data to the mesh description
+		bIsMeshDescriptionModified = true;
+		const bool bHasNormal = MorphTargetLodData.MorphNormalDeltas.Num() > 0;
+		SkeletalMeshAttributes.RegisterMorphTargetAttribute(MorphTargetName, bHasNormal);
+		TVertexAttributesRef<FVector3f> VertexPositionDeltas = SkeletalMeshAttributes.GetVertexMorphPositionDelta(MorphTargetName);
+		for (int32 VertexIndex = 0; VertexIndex < MorphTargetLodData.MorphPositionDeltas.Num(); ++VertexIndex)
+		{
+			FVertexID VertexID(VertexIndex);
+			VertexPositionDeltas[VertexID] = MorphTargetLodData.MorphPositionDeltas[VertexIndex];
+		}
+		if (bHasNormal)
+		{
+			TVertexInstanceAttributesRef<FVector3f> VertexInstanceNormalDeltas = SkeletalMeshAttributes.GetVertexInstanceMorphNormalDelta(MorphTargetName);
+			for (int32 VertexInstanceIndex = 0; VertexInstanceIndex < MorphTargetLodData.MorphPositionDeltas.Num(); ++VertexInstanceIndex)
+			{
+				FVertexInstanceID VertexInstanceID(VertexInstanceIndex);
+				VertexInstanceNormalDeltas[VertexInstanceID] = MorphTargetLodData.MorphNormalDeltas[VertexInstanceIndex];
+			}
+		}
+	}
+
+	return bIsMeshDescriptionModified;
+}
+
 void FLODUtilities::AdjustImportDataFaceMaterialIndex(const TArray<FSkeletalMaterial>& Materials, TArray<SkeletalMeshImportData::FMaterial>& RawMeshMaterials, TArray<SkeletalMeshImportData::FMeshFace>& LODFaces, int32 LODIndex)
 {
 	if ((Materials.Num() <= 1 && RawMeshMaterials.Num() <= 1) || LODIndex != 0)
diff --git a/Engine/Source/Developer/SkeletalMeshUtilitiesCommon/Public/LODUtilities.h b/Engine/Source/Developer/SkeletalMeshUtilitiesCommon/Public/LODUtilities.h
--- a/Engine/Source/Developer/SkeletalMeshUtilitiesCommon/Public/LODUtilities.h
+++ b/Engine/Source/Developer/SkeletalMeshUtilitiesCommon/Public/LODUtilities.h
@@ -243,6 +243,23 @@ public:
 	static void RestoreClothingFromBackup(USkeletalMesh* SkeletalMesh, TArray<ClothingAssetUtils::FClothingAssetMeshBinding>& ClothingBindings);
 	static void RestoreClothingFromBackup(USkeletalMesh* SkeletalMesh, TArray<ClothingAssetUtils::FClothingAssetMeshBinding>& ClothingBindings, const int32 LODIndex);
 
+	struct FMorphTargetLodBackupData
+	{
+		bool bIsEmpty = true;
+		TArray<FVector3f> MorphPositionDeltas;
+		TArray<FVector3f> MorphNormalDeltas;
+	};
+
+	/**
+	 * Backup all the custom imported morph targets.
+	 */
+	static void BackupCustomImportedMorphTargetData(USkeletalMesh* SkeletalMesh, TMap<FString, TArray<FMorphTargetLodBackupData>>& BackupImportedMorphTargetData);
+
+	/**
+	 * Restore the custom imported morph target in the mesh description. Return true if the mesh description was changed because it restore some morph targets
+	 */
+	static bool RestoreCustomImportedMorphTargetData(USkeletalMesh* SkeletalMesh, const int32 LodIndex, FMeshDescription& LodMeshDescription, const TMap<FString, TArray<FMorphTargetLodBackupData>>& BackupImportedMorphTargetData);
+
 	/**
 	 * Before building skeletalmesh base LOD (LOD index 0) using MeshUtilities.BuildSkeletalMesh, we want to adjust the imported faces material index to point on the correct sk material. We use the material name to match the material.
 	 * @param Materials - The skeletalmesh material list to fit the import data face material
diff --git a/Engine/Source/Runtime/SkeletalMeshDescription/Private/SkeletalMeshAttributes.cpp b/Engine/Source/Runtime/SkeletalMeshDescription/Private/SkeletalMeshAttributes.cpp
--- a/Engine/Source/Runtime/SkeletalMeshDescription/Private/SkeletalMeshAttributes.cpp
+++ b/Engine/Source/Runtime/SkeletalMeshDescription/Private/SkeletalMeshAttributes.cpp
@@ -478,6 +478,16 @@ TVertexInstanceAttributesConstRef<FVector3f> FSkeletalMeshAttributesShared::GetV
 	return MeshDescriptionShared.VertexInstanceAttributes().GetAttributesRef<FVector3f>(CreateMorphTargetAttributeName(InMorphTargetName));
 }
 
+bool FSkeletalMeshAttributesShared::HasMorphTargetPositionsAttribute(const FName InMorphTargetName) const
+{
+	return MeshDescriptionShared.VertexAttributes().HasAttribute(CreateMorphTargetAttributeName(InMorphTargetName));
+}
+
+bool FSkeletalMeshAttributesShared::HasMorphTargetNormalsAttribute(const FName InMorphTargetName) const
+{
+	return MeshDescriptionShared.VertexInstanceAttributes().HasAttribute(CreateMorphTargetAttributeName(InMorphTargetName));
+}
+
 
 bool FSkeletalMeshAttributesShared::HasBoneColorAttribute() const
 {
diff --git a/Engine/Source/Runtime/SkeletalMeshDescription/Public/SkeletalMeshAttributes.h b/Engine/Source/Runtime/SkeletalMeshDescription/Public/SkeletalMeshAttributes.h
--- a/Engine/Source/Runtime/SkeletalMeshDescription/Public/SkeletalMeshAttributes.h
+++ b/Engine/Source/Runtime/SkeletalMeshDescription/Public/SkeletalMeshAttributes.h
@@ -135,6 +135,10 @@ public:
 	/// normal vector for this vertex instance.
 	/// If the morph was not registered to include normals, this will return an invalid attribute. 
 	SKELETALMESHDESCRIPTION_API TVertexInstanceAttributesConstRef<FVector3f> GetVertexInstanceMorphNormalDelta(const FName InMorphTargetName) const;
+
+	SKELETALMESHDESCRIPTION_API bool HasMorphTargetPositionsAttribute(const FName InMorphTargetName) const;
+
+	SKELETALMESHDESCRIPTION_API bool HasMorphTargetNormalsAttribute(const FName InMorphTargetName) const;
 	
 	//
 	// Bones Methods
