Skeletal Mesh: Fix up material slot names on the mesh description if they don't point to a valid slot

#rnx

diff --git a/Engine/Source/Developer/SkeletalMeshUtilitiesCommon/Private/StaticToSkeletalMeshConverter.cpp b/Engine/Source/Developer/SkeletalMeshUtilitiesCommon/Private/StaticToSkeletalMeshConverter.cpp
--- a/Engine/Source/Developer/SkeletalMeshUtilitiesCommon/Private/StaticToSkeletalMeshConverter.cpp
+++ b/Engine/Source/Developer/SkeletalMeshUtilitiesCommon/Private/StaticToSkeletalMeshConverter.cpp
@@ -501,6 +501,16 @@ bool FStaticToSkeletalMeshConverter::InitializeSkeletalMeshFromMeshDescriptions(
 	// than materials in any of the LODs. Not the best system, but the best we have for now.
 	InSkeletalMesh->SetMaterials(TArray<FSkeletalMaterial>{InMaterials});
 
+	TSet<FName> ValidMaterialSlotNames;
+	for (int32 Index = 0; Index < InMaterials.Num(); Index++)
+	{
+		const FSkeletalMaterial& Material = InMaterials[Index];
+		if (!Material.MaterialSlotName.IsNone())
+		{
+			ValidMaterialSlotNames.Add(Material.MaterialSlotName);
+		}
+	}
+	
 	// This ensures that the render data gets built before we return, by calling PostEditChange when we fall out of scope.
 	{
 		FScopedSkeletalMeshPostEditChange ScopedPostEditChange( InSkeletalMesh );
@@ -526,6 +536,22 @@ bool FStaticToSkeletalMeshConverter::InitializeSkeletalMeshFromMeshDescriptions(
 			SkeletalLODInfo.BuildSettings.bRecomputeTangents = bInRecomputeTangents;
 
 			FMeshDescription ClonedDescription(*MeshDescription);
+
+			// Fix up the material slot names on the mesh to match the ones in the material list. If the name is
+			// either NAME_None, or doesn't exist in the material list, we use the group index to index into the
+			// material list to resolve the name.
+			FSkeletalMeshAttributes Attributes(ClonedDescription);
+			TPolygonGroupAttributesRef<FName> MaterialSlotNamesAttribute = Attributes.GetPolygonGroupMaterialSlotNames();
+			for (FPolygonGroupID PolygonGroupID: ClonedDescription.PolygonGroups().GetElementIDs())
+			{
+				if (!ValidMaterialSlotNames.Contains(MaterialSlotNamesAttribute.Get(PolygonGroupID)))
+				{
+					int32 MaterialIndex = PolygonGroupID.GetValue();
+					MaterialIndex = FMath::Clamp(MaterialIndex, 0, InMaterials.Num() - 1);
+					MaterialSlotNamesAttribute.Set(PolygonGroupID, InMaterials[MaterialIndex].MaterialSlotName);
+				}
+			}
+			
 			if (!AddLODFromMeshDescription(MoveTemp(ClonedDescription), InSkeletalMesh, MeshUtilities))
 			{
 				// If we didn't get a model for LOD index 0, we don't have a mesh. Bail out.
