[MorphTarget] SkeletalMesh ReImport Crash fix.

#rnx

diff --git a/Engine/Source/Developer/SkeletalMeshUtilitiesCommon/Private/LODUtilities.cpp b/Engine/Source/Developer/SkeletalMeshUtilitiesCommon/Private/LODUtilities.cpp
--- a/Engine/Source/Developer/SkeletalMeshUtilitiesCommon/Private/LODUtilities.cpp
+++ b/Engine/Source/Developer/SkeletalMeshUtilitiesCommon/Private/LODUtilities.cpp
@@ -4045,9 +4045,27 @@ bool FLODUtilities::RestoreCustomImportedMorphTargetData(USkeletalMesh* Skeletal
 			FName MorphTargetName(*MorphTargetNameStr);
 			if(UMorphTarget* MorphTarget = SkeletalMesh->FindMorphTarget(MorphTargetName))
 			{
-				MorphTarget->RemoveFromRoot();
-				MorphTarget->ClearFlags(RF_Standalone);
-				SkeletalMesh->UnregisterMorphTarget(MorphTarget, false);
+				if (MorphTarget)
+				{
+					MorphTarget->RemoveFromRoot();
+					MorphTarget->ClearFlags(RF_Standalone);
+
+					if (SkeletalMesh->HasMeshDescription(LodIndex))
+					{
+						//Remove the morph target from the raw import data
+						FMeshDescription* MeshDescription = SkeletalMesh->GetMeshDescription(LodIndex);
+						FSkeletalMeshAttributes MeshAttributes(*MeshDescription);
+
+						if (MeshAttributes.GetMorphTargetNames().Contains(MorphTargetName))
+						{
+							SkeletalMesh->ModifyMeshDescription(LodIndex);
+							MeshAttributes.UnregisterMorphTargetAttribute(MorphTargetName);
+							SkeletalMesh->CommitMeshDescription(LodIndex);
+						}
+					}
+
+					SkeletalMesh->UnregisterMorphTarget(MorphTarget, false);
+				}
 			}
 			if (FSkeletalMeshLODInfo* LodInfo = SkeletalMesh->GetLODInfo(LodIndex))
 			{
@@ -4073,9 +4091,7 @@ bool FLODUtilities::RestoreCustomImportedMorphTargetData(USkeletalMesh* Skeletal
 		const FMorphTargetLodBackupData& MorphTargetLodData = MorphTargetLodDatas[LodIndex];
 		if (MorphTargetLodData.bIsEmpty)
 		{
-			RemoveInvalidMorphTargetLOD(MorphTargetNameStr);
-			//Error out that the morph target is not valid anymore for this LOD
-			UE_ASSET_LOG(LogLODUtilities, Error, SkeletalMesh, TEXT("Cannot keep the generated morph target %s for LOD %d, because this LOD is now imported and the topology have change."), *MorphTargetNameStr, LodIndex);
+			//In case the Source Skeletal didn't have MorphTarget for a given LodIndex, it would show up as an Empty.
 			continue;
 		}
 
diff --git a/Engine/Source/Editor/Persona/Private/SMorphTargetViewer.cpp b/Engine/Source/Editor/Persona/Private/SMorphTargetViewer.cpp
--- a/Engine/Source/Editor/Persona/Private/SMorphTargetViewer.cpp
+++ b/Engine/Source/Editor/Persona/Private/SMorphTargetViewer.cpp
@@ -369,6 +369,8 @@ void SMorphTargetViewer::Construct(const FArguments& InArgs, const TSharedRef<IP
 
 	const FText SkeletalMeshName = SkeletalMesh ? FText::FromString( SkeletalMesh->GetName() ) : LOCTEXT( "MorphTargetMeshNameLabel", "No Skeletal Mesh Present" );
 
+	SkeletalMesh->GetOnMeshChanged().Add(FSimpleDelegate::CreateSP(this, &SMorphTargetViewer::OnMeshChanged));
+
 	ChildSlot
 	[
 		SNew( SVerticalBox )
@@ -822,6 +824,12 @@ void SMorphTargetViewer::OnPostUndo()
 	NotifySelectionChange();
 }
 
+void SMorphTargetViewer::OnMeshChanged()
+{
+	CreateMorphTargetList();
+	NotifySelectionChange();
+}
+
 void SMorphTargetViewer::NotifySelectionChange() const
 {
 	TArray< TSharedPtr< FDisplayedMorphTargetInfo > > SelectedRows = MorphTargetListView->GetSelectedItems();
diff --git a/Engine/Source/Editor/Persona/Private/SMorphTargetViewer.h b/Engine/Source/Editor/Persona/Private/SMorphTargetViewer.h
--- a/Engine/Source/Editor/Persona/Private/SMorphTargetViewer.h
+++ b/Engine/Source/Editor/Persona/Private/SMorphTargetViewer.h
@@ -163,6 +163,11 @@ public:
 	 */
 	void OnPostUndo();
 
+	/**
+	 * Refreshes the morph target list after a mesh changed
+	 */
+	void OnMeshChanged();
+
 private:
 
 	/** Handler for context menus */
diff --git a/Engine/Source/Runtime/SkeletalMeshDescription/Private/SkeletalMeshAttributes.cpp b/Engine/Source/Runtime/SkeletalMeshDescription/Private/SkeletalMeshAttributes.cpp
--- a/Engine/Source/Runtime/SkeletalMeshDescription/Private/SkeletalMeshAttributes.cpp
+++ b/Engine/Source/Runtime/SkeletalMeshDescription/Private/SkeletalMeshAttributes.cpp
@@ -214,6 +214,13 @@ bool FSkeletalMeshAttributes::UnregisterMorphTargetAttribute(const FName InMorph
 	}
 
 	MeshDescription.VertexAttributes().UnregisterAttribute(AttributeName);
+
+	// VertexInstancesAttributes
+	if (MeshDescription.VertexInstanceAttributes().HasAttribute(AttributeName))
+	{
+		MeshDescription.VertexInstanceAttributes().UnregisterAttribute(AttributeName);
+	}
+
 	return true;
 }
 
