Alternate Skin Weights fixes
    - MeshAssetFunctions: updated ddc cache when adding a new profile (should be handled in CommitMeshDescription)
    - LODUtilities: corrected GenerateImportedSkinWeightProfileData, which did not correctly handle maximum bone influence + weight sorting (+ fixed similar position test)
    - SkinnedMeshComponent: added GSkinnedMeshSkinWeightProfileEarlyExit (false by default) to avoid exiting too early when defining a new profile stack
    - FSkeletalMeshAttributes: adding a new skin weight attribute now returns true if the attribute is already available

NOTE: things to keep in mind when generating rendering meshes from editable meshes (mainly for the future me)
    - influences may be clamped at 8/12 when converted for rendering, even if editable data can support any number (be sure to check rendering settings and the number of bones supported by the current platform)
    - as weights can be clamped, it's important that weights are sorted from highest to lowest
    - use FBoneWeights to ensure that weights are re-normalized if they are clamped

TODO:
    - test CommitMeshDescription to ensure that adding a new profile triggers a DDC invalidation
    - there is redundancy between MeshUtilities::GenerateSkeletalRenderMesh and LOD rendering mesh generation (which is a source of errors)
    - FLODUtilities::UpdateAlternateSkinWeights could benefit from a more robust projection algorithm (see FTransferBoneWeights)

#rnx

diff --git a/Engine/Plugins/Runtime/GeometryScripting/Source/GeometryScriptingCore/Private/MeshAssetFunctions.cpp b/Engine/Plugins/Runtime/GeometryScripting/Source/GeometryScriptingCore/Private/MeshAssetFunctions.cpp
--- a/Engine/Plugins/Runtime/GeometryScripting/Source/GeometryScriptingCore/Private/MeshAssetFunctions.cpp
+++ b/Engine/Plugins/Runtime/GeometryScripting/Source/GeometryScriptingCore/Private/MeshAssetFunctions.cpp
@@ -1379,6 +1379,7 @@ UDynamicMesh* UGeometryScriptLibrary_StaticMeshFunctions::CopySkinWeightProfileT
 	}
 
 	ToSkeletalMeshAsset->CommitMeshDescription(TargetLOD.LODIndex);
+	ToSkeletalMeshAsset->InvalidateDeriveDataCacheGUID();
 
 	if (Options.bDeferMeshPostEditChange == false)
 	{
diff --git a/Engine/Source/Developer/SkeletalMeshUtilitiesCommon/Private/LODUtilities.cpp b/Engine/Source/Developer/SkeletalMeshUtilitiesCommon/Private/LODUtilities.cpp
--- a/Engine/Source/Developer/SkeletalMeshUtilitiesCommon/Private/LODUtilities.cpp
+++ b/Engine/Source/Developer/SkeletalMeshUtilitiesCommon/Private/LODUtilities.cpp
@@ -2498,7 +2498,7 @@ bool FLODUtilities::UpdateAlternateSkinWeights(
 		ZBound.Z = PositionZ.Z - ComparisonThreshold;
 		const int32 StartRange = Algo::LowerBound(VertIndexAndZ, ZBound, FCompareIndexAndZ{});
 
-		ZBound.Z = Position.Z + ComparisonThreshold;
+		ZBound.Z = PositionZ.Z + ComparisonThreshold;
 		const int32 EndRange = Algo::UpperBound(VertIndexAndZ, ZBound, FCompareIndexAndZ{});
 		
 		// Search for duplicates, quickly!
@@ -2974,6 +2974,8 @@ void FLODUtilities::GenerateImportedSkinWeightProfileData(
 	int32 BoneInfluenceLimit,
 	const ITargetPlatform* TargetPlatform)
 {
+	using namespace UE::AnimationCore;
+	
 	//Add the override buffer with the alternate influence data
 	TArray<FSoftSkinVertex> DestinationSoftVertices;
 	LODModelDest.GetVertices(DestinationSoftVertices);
@@ -2997,6 +2999,18 @@ void FLODUtilities::GenerateImportedSkinWeightProfileData(
 	TMap<FBoneIndexType, float> WeightForBone;
 	TArray<const SkeletalMeshImportData::FVertInfluence*> FoundInfluences;
 
+	// cf. FCompareVertexIndex (should be factorized for next versions) 
+	auto VertInfluencePredicate = [](const SkeletalMeshImportData::FVertInfluence* A, const SkeletalMeshImportData::FVertInfluence* B) 
+	{
+		if (A->VertIndex > B->VertIndex) return false;
+		if (A->VertIndex < B->VertIndex) return true;
+		if (A->Weight < B->Weight) return false;
+		if (A->Weight > B->Weight) return true;
+		if (A->BoneIndex > B->BoneIndex) return false;
+		if (A->BoneIndex < B->BoneIndex) return true;
+		return  false;
+	};
+	
 	for (int32 VertexInstanceIndex = 0; VertexInstanceIndex < DestinationSoftVertices.Num(); ++VertexInstanceIndex)
 	{
 		int32 SectionIndex = INDEX_NONE;
@@ -3006,18 +3020,19 @@ void FLODUtilities::GenerateImportedSkinWeightProfileData(
 		{
 			continue;
 		}
+		
 		FSkelMeshSection& Section = LODModelDest.Sections[SectionIndex];
 		const TArray<FBoneIndexType>& SectionBoneMap = Section.BoneMap;
-		const FSoftSkinVertex& Vertex = DestinationSoftVertices[VertexInstanceIndex];
 		const int32 VertexIndex = LODModelDest.MeshToImportVertexMap[VertexInstanceIndex];
 		check(VertexIndex >= 0 && VertexIndex <= LODModelDest.MaxImportVertex);
 		FRawSkinWeight& SkinWeight = SkinWeights.AddDefaulted_GetRef();
 
-		WeightForBone.Reset();
+		// get the bone weights for that vertex and sort them (mostly by weights here as VertIndex are similar)
 		FoundInfluences.Reset();
-		
 		VertexToInfluenceMap.MultiFind(VertexIndex, FoundInfluences);
-		
+		Algo::Sort(FoundInfluences, VertInfluencePredicate);
+
+		WeightForBone.Reset();
 		for (const SkeletalMeshImportData::FVertInfluence* VertInfluence : FoundInfluences)
 		{
 			//Use the section bone map to remap the bone index
@@ -3031,51 +3046,50 @@ void FLODUtilities::GenerateImportedSkinWeightProfileData(
 			WeightForBone.Add(static_cast<FBoneIndexType>(BoneMapIndex), VertInfluence->Weight);
 		}
 
-		using namespace UE::AnimationCore;
-
-		//Add the prepared alternate influences for this skin vertex
-		int32 InfluenceBoneCount = 0;
+		//Add the prepared alternate influences for this skin vertex (cf. MeshUtilities::GenerateSkeletalRenderMesh)
+		int32 InfluenceBoneCount = FMath::Min<uint32>(WeightForBone.Num(), MaxBoneInfluencesFromAsset);
+		if (InfluenceBoneCount > EXTRA_BONE_INFLUENCES && !FGPUBaseSkinVertexFactory::UseUnlimitedBoneInfluences(InfluenceBoneCount, TargetPlatform))
+		{
+			InfluenceBoneCount = EXTRA_BONE_INFLUENCES;
+		}
 		
 		FBoneIndexType InfluenceBones[MAX_TOTAL_INFLUENCES];
 		float InfluenceWeights[MAX_TOTAL_INFLUENCES];
-		
-		for (const TTuple<FBoneIndexType, float>& Kvp : WeightForBone)
+
+		int32 ActualInfluenceCount = 0;
+		for (const TTuple<FBoneIndexType, float>& BoneAndWeight : WeightForBone)
 		{
-			InfluenceBones[InfluenceBoneCount] = Kvp.Key;
-			InfluenceWeights[InfluenceBoneCount] = Kvp.Value;
-			InfluenceBoneCount++;
-			//Do not add more influence then the maximum we support on the platform or for this asset
-			if (InfluenceBoneCount >= MaxInfluenceCount)
+			if (ActualInfluenceCount < InfluenceBoneCount)
 			{
-				break;
+				InfluenceBones[ActualInfluenceCount] = BoneAndWeight.Key;
+				InfluenceWeights[ActualInfluenceCount] = BoneAndWeight.Value;
+				ActualInfluenceCount++;
 			}
 		}
 
-		const FBoneWeights BoneWeights = FBoneWeights::Create(InfluenceBones, InfluenceWeights, InfluenceBoneCount);
+		const FBoneWeights BoneWeights = FBoneWeights::Create(InfluenceBones, InfluenceWeights, ActualInfluenceCount);
 		FMemory::Memzero(SkinWeight.InfluenceBones);
 		FMemory::Memzero(SkinWeight.InfluenceWeights);
 	
 		if (BoneWeights.Num() == 0)
 		{
 			SkinWeight.InfluenceWeights[0] = std::numeric_limits<uint16>::max();
-			InfluenceBoneCount = 1;
 		}
 		else
 		{
 			int32 Index = 0;
-			for (FBoneWeight BoneWeight: BoneWeights)
+			for (const FBoneWeight& BoneWeight: BoneWeights)
 			{
 				SkinWeight.InfluenceBones[Index] = BoneWeight.GetBoneIndex();
 				SkinWeight.InfluenceWeights[Index] = BoneWeight.GetRawWeight();
 				Index++;
 			}
-			InfluenceBoneCount = BoneWeights.Num(); 
 		}
 
 		//Adjust section influence count if the alternate influence bone count is greater
-		if (InfluenceBoneCount > MaxNumInfluences)
+		if (ActualInfluenceCount > MaxNumInfluences)
 		{
-			MaxNumInfluences = InfluenceBoneCount;
+			MaxNumInfluences = ActualInfluenceCount;
 			if (MaxNumInfluences > Section.GetMaxBoneInfluences())
 			{
 				Section.MaxBoneInfluences = MaxNumInfluences;
diff --git a/Engine/Source/Runtime/Engine/Private/Components/SkinnedMeshComponent.cpp b/Engine/Source/Runtime/Engine/Private/Components/SkinnedMeshComponent.cpp
--- a/Engine/Source/Runtime/Engine/Private/Components/SkinnedMeshComponent.cpp
+++ b/Engine/Source/Runtime/Engine/Private/Components/SkinnedMeshComponent.cpp
@@ -104,6 +104,13 @@ static FAutoConsoleVariableRef CVarSkinnedMeshRenderNanite(
 	)
 );
 
+static bool GSkinnedMeshSkinWeightProfileEarlyExit = false;
+static FAutoConsoleVariableRef CVarSkinnedMeshSkinWeightProfileEarlyExit(
+	TEXT("r.SkinnedMesh.SkinWeightProfileEarlyExit"),
+   GSkinnedMeshSkinWeightProfileEarlyExit,
+   TEXT("When set to true, the new profile stack update is not re-done if it's similar to the current skin weight profile.")
+   );
+
 bool ShouldRenderNaniteSkinnedMeshes()
 {
 	return NaniteSkinnedMeshesSupported() && GSkinnedMeshRenderNanite != 0;
@@ -5040,7 +5047,10 @@ bool USkinnedMeshComponent::SetSkinWeightProfileStack(const FSkinWeightProfileSt
 	// If we're not actually changing anything, then just return and say we did, as if the operation succeeded.
 	if (InProfileStack.Normalized() == FSkinWeightProfileStack{CurrentSkinWeightProfileLayers}.Normalized())
 	{
-		return true;
+		if (GSkinnedMeshSkinWeightProfileEarlyExit)
+		{
+			return true;
+		}
 	}
 
 	bool bChanged = false;
@@ -5155,6 +5165,11 @@ bool USkinnedMeshComponent::SetSkinWeightProfileStack(const FSkinWeightProfileSt
 					// Put in a skin weight profile request
 					if (FSkinWeightProfileManager* Manager = FSkinWeightProfileManager::Get(GetWorld()))
 					{
+						if (bSkinWeightProfilePending)
+						{
+							Manager->CancelSkinWeightProfileRequest(this);	
+						}
+
 						Manager->RequestSkinWeightProfileStack(InProfileStack, GetSkinnedAsset(), this, Callback);
 						bSkinWeightProfilePending = true;
 					}
diff --git a/Engine/Source/Runtime/SkeletalMeshDescription/Private/SkeletalMeshAttributes.cpp b/Engine/Source/Runtime/SkeletalMeshDescription/Private/SkeletalMeshAttributes.cpp
--- a/Engine/Source/Runtime/SkeletalMeshDescription/Private/SkeletalMeshAttributes.cpp
+++ b/Engine/Source/Runtime/SkeletalMeshDescription/Private/SkeletalMeshAttributes.cpp
@@ -115,7 +115,7 @@ bool FSkeletalMeshAttributes::RegisterSkinWeightAttribute(const FName InProfileN
 	
 	if (MeshDescription.VertexAttributes().HasAttribute(AttributeName))
 	{
-		return false;
+		return true;
 	}
 
 	return MeshDescription.VertexAttributes().RegisterAttribute<int32[]>(AttributeName, 1, 0, EMeshAttributeFlags::None).IsValid();
