Alternative Skin Weight: unused vertices taken into account when generating render data

NOTE:
    when the skeletal mesh contains unused vertices (which are not referenced by any face), we avoid using them in asw calculations to prevent generating inconsistent rendering data.
    in this scenario, faces are used instead to determine which vertices we're interested in, and use their information to generate more consistent weights for the unused ones.

diff --git a/Engine/Plugins/Runtime/GeometryScripting/Source/GeometryScriptingCore/Private/MeshAssetFunctions.cpp b/Engine/Plugins/Runtime/GeometryScripting/Source/GeometryScriptingCore/Private/MeshAssetFunctions.cpp
--- a/Engine/Plugins/Runtime/GeometryScripting/Source/GeometryScriptingCore/Private/MeshAssetFunctions.cpp
+++ b/Engine/Plugins/Runtime/GeometryScripting/Source/GeometryScriptingCore/Private/MeshAssetFunctions.cpp
@@ -788,6 +788,13 @@ UDynamicMesh* UGeometryScriptLibrary_StaticMeshFunctions::CopyMeshFromSkeletalMe
 		Converter.bVIDsFromNonManifoldMeshDescriptionAttr = true;
 		
 		Converter.Convert(SourceMesh, NewMesh, AssetOptions.bRequestTangents);
+
+		// Notify of unused vertices
+		if (NewMesh.HasUnusedVertices())
+		{
+			static const FText UnusedVerticesMsg = LOCTEXT("CopyMeshFromSkeletalMesh_UnusedVertices", "CopyMeshFromSkeletalMesh: ToDynamicMesh has unused vertices (not referenced by any triangle)");
+			AppendWarning(Debug, EGeometryScriptErrorType::UnknownError, UnusedVerticesMsg);	
+		}
 	
 		ToDynamicMesh->SetMesh(MoveTemp(NewMesh));
 	
diff --git a/Engine/Source/Developer/SkeletalMeshUtilitiesCommon/Private/LODUtilities.cpp b/Engine/Source/Developer/SkeletalMeshUtilitiesCommon/Private/LODUtilities.cpp
--- a/Engine/Source/Developer/SkeletalMeshUtilitiesCommon/Private/LODUtilities.cpp
+++ b/Engine/Source/Developer/SkeletalMeshUtilitiesCommon/Private/LODUtilities.cpp
@@ -2429,6 +2429,315 @@ static void InfillMissingVertexSkinWeights(
 	}
 }
 
+namespace AlternateSkinWeightsLocal
+{
+
+void NotifyFromFacesDiscrepancy(const USkeletalMesh* InSkeletalMeshDest, const int32 InNumSrcFaceVerts, const int32 InNumDstFaceVerts)
+{
+	FFormatNamedArguments Args;
+	Args.Add(TEXT("SkeletalMeshName"), FText::FromString(InSkeletalMeshDest->GetName()));
+	Args.Add(TEXT("FacesPointNumberSrc"), InNumSrcFaceVerts);
+	Args.Add(TEXT("FacesPointNumberDest"), InNumDstFaceVerts);
+
+	FText Message = FText::Format(NSLOCTEXT("FLODUtilities_UpdateAlternateSkinWeights", "DifferentFaceVerticesNumber",
+		"Asset {SkeletalMeshName} failed to import skin weight profile the alternate skinning model has a different number of face vertex."
+		" Alternate face vertex count: {FacesPointNumberSrc}, LOD vertex count: {FacesPointNumberDest}"), Args);
+	
+	UE_LOG(LogLODUtilities, Warning, TEXT("%s"), *(Message.ToString()));
+}
+
+void GetVerticesFromFaces(const FSkeletalMeshImportData& InImportData, TSet<uint32>& OutFaceVertices)
+{
+	const TConstArrayView<SkeletalMeshImportData::FTriangle> Faces(InImportData.Faces);
+	const TConstArrayView<SkeletalMeshImportData::FVertex> Wedges(InImportData.Wedges);
+	
+	OutFaceVertices.Empty(Faces.Num());
+	
+	for (const SkeletalMeshImportData::FTriangle& Triangle: Faces)
+	{
+		OutFaceVertices.Add(Wedges[Triangle.WedgeIndex[0]].VertexIndex);
+		OutFaceVertices.Add(Wedges[Triangle.WedgeIndex[1]].VertexIndex);
+		OutFaceVertices.Add(Wedges[Triangle.WedgeIndex[2]].VertexIndex);
+	}
+ 
+	OutFaceVertices.Sort([](uint32 VertIndex0, uint32 VertIndex1)
+	{
+		return VertIndex0 < VertIndex1;
+	});		
+}
+
+// note that this algo supposes that InUsed is sorted
+void GetUnusedVertices(const TSet<uint32>& InUsed, const int32 InExpectedUnusedNum, TSet<uint32>& OutUnused)
+{
+	OutUnused.Empty(InExpectedUnusedNum);
+
+	int32 UnusedVertIndex = 0;
+	for (int32 UsedVertIndex: InUsed)
+	{
+		while (UnusedVertIndex < UsedVertIndex)
+		{
+			OutUnused.Add(UnusedVertIndex++);
+		}
+		UnusedVertIndex = UsedVertIndex+1;
+	}
+	
+	ensure(OutUnused.Num() == InExpectedUnusedNum);
+}
+	
+class FVertexMatcher
+{
+public:
+	
+	FVertexMatcher(
+		const FSkeletalMeshImportData& InImportDataSrc,
+		const FSkeletalMeshImportData& InImportDataDst,
+		const TSet<uint32>& InDstVertices,
+		TArray<uint32>& InOutVertexIndexToMatchWithPositions,
+		TSortedMap<uint32, VertexMatchNameSpace::FVertexMatchResult>& InOutVertexIndexSrcToVertexIndexDestMatches)
+	: ImportDataSrc(InImportDataSrc)
+	, ImportDataDst(InImportDataDst)
+	, DstVertices(InDstVertices)
+	, VertexIndexToMatchWithPositions(InOutVertexIndexToMatchWithPositions)
+	, VertexIndexSrcToVertexIndexDestMatches(InOutVertexIndexSrcToVertexIndexDestMatches)
+	{
+		for (int32 WedgeIndex = 0; WedgeIndex < ImportDataSrc.Wedges.Num(); WedgeIndex++)
+		{
+			const SkeletalMeshImportData::FVertex& Wedge = ImportDataSrc.Wedges[WedgeIndex];
+			VertexToSrcWedgeMap.Add(Wedge.VertexIndex, WedgeIndex);
+		}
+		for (int32 WedgeIndex = 0; WedgeIndex < ImportDataDst.Wedges.Num(); WedgeIndex++)
+		{
+			const SkeletalMeshImportData::FVertex& Wedge = ImportDataDst.Wedges[WedgeIndex];
+			VertexToDstWedgeMap.Add(Wedge.VertexIndex, WedgeIndex);
+		}
+	}
+
+	void TryMatching(const TSet<uint32>& InSrcVertices) const
+	{
+		// Create a list of vertex Z/index pairs
+		TArray<FIndexAndZ> VertIndexAndZ;
+		VertIndexAndZ.Reserve(DstVertices.Num());
+		for (int32 VertexIndex: DstVertices)
+		{
+			new(VertIndexAndZ)FIndexAndZ(VertexIndex, ImportDataDst.Points[VertexIndex]);
+		}
+		// Project the vertices onto a diagonal, use that to quickly find vertices for
+		// closest point matching. 
+		VertIndexAndZ.Sort(FCompareIndexAndZ());
+		
+		auto FindSimilarPosition = [&VertIndexAndZ, this](const FVector3f& Position, TArray<int32>& PositionMatches, const float ComparisonThreshold)
+		{
+			PositionMatches.Reset();
+			const FIndexAndZ PositionZ = FIndexAndZ(0, Position);
+		
+			// Use binary search to narrow down the range that matches our search area.
+			FIndexAndZ ZBound;
+			ZBound.Z = PositionZ.Z - ComparisonThreshold;
+			const int32 StartRange = Algo::LowerBound(VertIndexAndZ, ZBound, FCompareIndexAndZ{});
+
+			ZBound.Z = PositionZ.Z + ComparisonThreshold;
+			const int32 EndRange = Algo::UpperBound(VertIndexAndZ, ZBound, FCompareIndexAndZ{});
+		
+			// Search for duplicates, quickly!
+			for (int32 Index = StartRange; Index < EndRange; Index++)
+			{
+				const FVector3f& PositionA = ImportDataDst.Points[VertIndexAndZ[Index].Index];
+				if (PointsEqual(PositionA, Position, ComparisonThreshold))
+				{
+					PositionMatches.Add(VertIndexAndZ[Index].Index);
+				}
+			}
+		};
+
+		TArray<int32> SrcWedgeIndexes;
+		TArray<int32> DestWedgeIndexes;
+		
+		// Match all source vertex with destination vertex
+		for (int32 VertexIndexSrc: InSrcVertices)
+		{
+			const FVector3f& PositionSrc = ImportDataSrc.Points[VertexIndexSrc];
+		
+			TArray<int32> SimilarDestinationVertex;
+			FindSimilarPosition(PositionSrc, SimilarDestinationVertex, KINDA_SMALL_NUMBER);
+
+			if (SimilarDestinationVertex.Num() == 0)
+			{
+				//Match with UV projection
+				VertexIndexToMatchWithPositions.Add(VertexIndexSrc);
+			}
+			else
+			{
+				//We have a direct match
+				VertexMatchNameSpace::FVertexMatchResult& VertexMatchDest = VertexIndexSrcToVertexIndexDestMatches.Add(VertexIndexSrc);
+
+				SrcWedgeIndexes.Reset();
+				VertexToSrcWedgeMap.MultiFind(VertexIndexSrc, SrcWedgeIndexes);
+
+				if (SrcWedgeIndexes.Num() > 0 && SimilarDestinationVertex.Num() > 1)
+				{
+					//Check if we have a point that is perfectly matching (position, UV, material and vertex color). Because normals and tangent are on the triangles we do not test those.
+					for (int32 MatchDestinationIndex = 0; MatchDestinationIndex < SimilarDestinationVertex.Num(); ++MatchDestinationIndex)
+					{
+						int32 VertexIndexDest = SimilarDestinationVertex[MatchDestinationIndex];
+						DestWedgeIndexes.Reset();
+						VertexToDstWedgeMap.MultiFind(VertexIndexDest, DestWedgeIndexes);
+					
+						for (int32 IndexDest = 0; IndexDest < DestWedgeIndexes.Num(); ++IndexDest)
+						{
+							int32 DestWedgeIndex = DestWedgeIndexes[IndexDest];
+							const SkeletalMeshImportData::FVertex& WedgeDest = ImportDataDst.Wedges[DestWedgeIndex];
+							for (int32 IndexSrc = 0; IndexSrc < SrcWedgeIndexes.Num(); ++IndexSrc)
+							{
+								int32 SrcWedgeIndex = SrcWedgeIndexes[IndexSrc];
+								const SkeletalMeshImportData::FVertex& WedgeSrc = ImportDataSrc.Wedges[SrcWedgeIndex];
+								//Wedge == operator test: material, vertex color and UVs
+								if (WedgeDest == WedgeSrc)
+								{
+									VertexMatchDest.VertexIndexes.Add(SimilarDestinationVertex[MatchDestinationIndex]);
+									VertexMatchDest.Ratios.Add(1.0f);
+									break;
+								}
+							}
+							if (VertexMatchDest.VertexIndexes.Num() > 0)
+							{
+								break;
+							}
+						}
+					}
+				}
+				
+				//If there is no direct match, simply put everything
+				if (VertexMatchDest.VertexIndexes.Num() == 0)
+				{
+					for (int32 MatchDestinationIndex = 0; MatchDestinationIndex < SimilarDestinationVertex.Num(); ++MatchDestinationIndex)
+					{
+						VertexMatchDest.VertexIndexes.Add(SimilarDestinationVertex[MatchDestinationIndex]);
+						VertexMatchDest.Ratios.Add(1.0f / SimilarDestinationVertex.Num());
+					}
+				}
+			}
+		}
+	}
+	
+private:
+	const FSkeletalMeshImportData& ImportDataSrc;
+	const FSkeletalMeshImportData& ImportDataDst;
+	const TSet<uint32>& DstVertices;
+	TMultiMap<int32, int32> VertexToSrcWedgeMap;
+	TMultiMap<int32, int32> VertexToDstWedgeMap;
+	TArray<uint32>& VertexIndexToMatchWithPositions;
+	TSortedMap<uint32, VertexMatchNameSpace::FVertexMatchResult>& VertexIndexSrcToVertexIndexDestMatches;
+};
+
+uint64 HashFromBoneAndVertex(const int32 InBoneIndex, const int32 InVertexIndex)
+{
+	return ((uint64)(InBoneIndex) << 32 & 0xFFFFFFFF00000000) | ((uint64)(InVertexIndex) & 0x00000000FFFFFFFF);
+}
+
+class FUnusedInfluenceBuilder
+{
+public:
+	FUnusedInfluenceBuilder(
+		const TSortedMap<uint32, VertexMatchNameSpace::FVertexMatchResult>& InVertexIndexSrcToVertexIndexDestMatches,
+		const TMap<uint32, TArray<int32>>& InSourceVertexIndexToAlternateInfluenceIndexMap,
+		const int32 InNumVertices,
+		TArray<SkeletalMeshImportData::FRawBoneInfluence>& InOutAlternateInfluences)
+		: VertexIndexSrcToVertexIndexDestMatches(InVertexIndexSrcToVertexIndexDestMatches)
+		, SourceVertexIndexToAlternateInfluenceIndexMap(InSourceVertexIndexToAlternateInfluenceIndexMap)
+		, NumVertices(InNumVertices)
+		, AlternateInfluences(InOutAlternateInfluences)	
+	{}
+	
+	void TryBuilding(const TSet<uint32>& InUnusedVertices) const
+	{
+		const int32 AlternateInfluenceNum = AlternateInfluences.Num();
+		if (AlternateInfluenceNum == 0 || NumVertices == 0 || InUnusedVertices.IsEmpty())
+		{
+			return;
+		}
+		
+		TArray<bool> SrcVertexIndexMatched;
+		SrcVertexIndexMatched.AddZeroed(NumVertices);
+
+		TMap<uint64, int32> InfluenceKeyToInfluenceIndex;
+		InfluenceKeyToInfluenceIndex.Reserve(AlternateInfluenceNum);
+	
+		for (int32 AlternateInfluenceIndex = 0; AlternateInfluenceIndex < AlternateInfluenceNum; ++AlternateInfluenceIndex)
+		{
+			SkeletalMeshImportData::FRawBoneInfluence& Influence = AlternateInfluences[AlternateInfluenceIndex];
+			SrcVertexIndexMatched[Influence.VertexIndex] = true;
+			const uint64 Key = HashFromBoneAndVertex(Influence.BoneIndex, Influence.VertexIndex);
+			InfluenceKeyToInfluenceIndex.Add(Key, AlternateInfluenceIndex);
+		}
+
+		for (const uint32 Unused: InUnusedVertices)
+		{
+			//Skip if the vertex is already matched
+			if (SrcVertexIndexMatched[Unused])
+			{
+				continue;
+			}
+
+			const VertexMatchNameSpace::FVertexMatchResult* SrcVertexMatch = VertexIndexSrcToVertexIndexDestMatches.Find(Unused);
+			if (!SrcVertexMatch || SrcVertexMatch->VertexIndexes.IsEmpty())
+			{
+				//No match skip this influence
+				continue;
+			}
+		
+			for (int32 ImpactedIndex = 0; ImpactedIndex < SrcVertexMatch->VertexIndexes.Num(); ++ImpactedIndex)
+			{
+				float Ratio = SrcVertexMatch->Ratios[ImpactedIndex];
+				if (!FMath::IsNearlyZero(Ratio, KINDA_SMALL_NUMBER))
+				{
+					uint32 VertexIndexSrc = SrcVertexMatch->VertexIndexes[ImpactedIndex];
+					const TArray<int32>* SrcAlternateInfluences = SourceVertexIndexToAlternateInfluenceIndexMap.Find(VertexIndexSrc);
+					if (!SrcAlternateInfluences)
+					{
+						continue;
+					}
+
+					for (const int32 AlternateInfluenceIndex: *SrcAlternateInfluences)
+					{
+						if (!AlternateInfluences.IsValidIndex(AlternateInfluenceIndex))
+						{
+							continue;
+						}
+					
+						SrcVertexIndexMatched[Unused] = true;
+						SkeletalMeshImportData::FRawBoneInfluence AlternateInfluence = AlternateInfluences[AlternateInfluenceIndex];
+						const uint64 Key = HashFromBoneAndVertex(AlternateInfluence.BoneIndex, Unused);
+					
+						if (!InfluenceKeyToInfluenceIndex.Contains(Key))
+						{
+							AlternateInfluence.VertexIndex = Unused;
+							InfluenceKeyToInfluenceIndex.Add(Key, AlternateInfluences.Add(AlternateInfluence));
+						}
+						else
+						{
+							int32& InfluenceIndex = InfluenceKeyToInfluenceIndex.FindOrAdd(Key);
+							SkeletalMeshImportData::FRawBoneInfluence& ExistAlternateInfluence = AlternateInfluences[InfluenceIndex];
+							if (ExistAlternateInfluence.Weight < AlternateInfluence.Weight)
+							{
+								ExistAlternateInfluence.Weight = AlternateInfluence.Weight;
+							}
+						}
+					}
+				}
+			}
+		}
+	}
+	
+private:
+	
+	const TSortedMap<uint32, VertexMatchNameSpace::FVertexMatchResult>& VertexIndexSrcToVertexIndexDestMatches;
+	const TMap<uint32, TArray<int32>>& SourceVertexIndexToAlternateInfluenceIndexMap;
+	int32 NumVertices = 0;
+	TArray<SkeletalMeshImportData::FRawBoneInfluence>& AlternateInfluences;
+};
+	
+}
+
 bool FLODUtilities::UpdateAlternateSkinWeights(
 	FSkeletalMeshLODModel& LODModelDest,
 	FSkeletalMeshImportData& ImportDataDest,
@@ -2438,10 +2747,11 @@ bool FLODUtilities::UpdateAlternateSkinWeights(
 	int32 LODIndexDest,
 	const IMeshUtilities::MeshBuildOptions& Options)
 {
+	using namespace AlternateSkinWeightsLocal;
+	
 	//Ensure log message only once
 	bool bNoMatchMsgDone = false;
-	int32 PointNumberDest = ImportDataDest.Points.Num();
-	int32 VertexNumberDest = ImportDataDest.Points.Num();
+	const int32 PointNumberDest = ImportDataDest.Points.Num();
 
 	int32 ProfileIndex = 0;
 	if (!ImportDataDest.AlternateInfluenceProfileNames.Find(ProfileNameDest.ToString(), ProfileIndex))
@@ -2461,9 +2771,8 @@ bool FLODUtilities::UpdateAlternateSkinWeights(
 	check(ImportDataDest.AlternateInfluences.IsValidIndex(ProfileIndex));
 	//The data must be there and must be verified before getting here
 	const FSkeletalMeshImportData& ImportDataSrc = ImportDataDest.AlternateInfluences[ProfileIndex];
-	int32 PointNumberSrc = ImportDataSrc.Points.Num();
-	int32 VertexNumberSrc = ImportDataSrc.Points.Num();
-	int32 InfluenceNumberSrc = ImportDataSrc.Influences.Num();
+	const int32 PointNumberSrc = ImportDataSrc.Points.Num();
+	const int32 InfluenceNumberSrc = ImportDataSrc.Influences.Num();
 
 	if (PointNumberDest != PointNumberSrc)
 	{
@@ -2477,140 +2786,54 @@ bool FLODUtilities::UpdateAlternateSkinWeights(
 		return false;
 	}
 
-	// Create a list of vertex Z/index pairs
-	TArray<FIndexAndZ> VertIndexAndZ;
-	VertIndexAndZ.Reserve(VertexNumberDest);
-	for (int32 VertexIndex = 0; VertexIndex < VertexNumberDest; ++VertexIndex)
-	{
-		new(VertIndexAndZ)FIndexAndZ(VertexIndex, ImportDataDest.Points[VertexIndex]);
-	}
-	// Project the vertices onto a diagonal, use that to quickly find vertices for
-	// closest point matching. 
-	VertIndexAndZ.Sort(FCompareIndexAndZ());
+	// to avoid building render data using 'unused' vertices (not referenced by any triangle),
+	// retrieve the vertices info from the triangles
+	TSet<uint32> SrcFaceVert;
+	GetVerticesFromFaces(ImportDataSrc, SrcFaceVert);
+	
+	TSet<uint32> DstFaceVert;
+	GetVerticesFromFaces(ImportDataDest, DstFaceVert);
 
-	auto FindSimilarPosition = [&VertIndexAndZ, &ImportDataDest](const FVector3f& Position, TArray<int32>& PositionMatches, const float ComparisonThreshold)
+	if (DstFaceVert.Num() != SrcFaceVert.Num())
 	{
-		PositionMatches.Reset();
-		const FIndexAndZ PositionZ = FIndexAndZ(0, Position);
-		
-		// Use binary search to narrow down the range that matches our search area.
-		FIndexAndZ ZBound;
-		ZBound.Z = PositionZ.Z - ComparisonThreshold;
-		const int32 StartRange = Algo::LowerBound(VertIndexAndZ, ZBound, FCompareIndexAndZ{});
-
-		ZBound.Z = PositionZ.Z + ComparisonThreshold;
-		const int32 EndRange = Algo::UpperBound(VertIndexAndZ, ZBound, FCompareIndexAndZ{});
-		
-		// Search for duplicates, quickly!
-		for (int32 Index = StartRange; Index < EndRange; Index++)
-		{
-			const FVector3f& PositionA = ImportDataDest.Points[VertIndexAndZ[Index].Index];
-			if (PointsEqual(PositionA, Position, ComparisonThreshold))
-			{
-				PositionMatches.Add(VertIndexAndZ[Index].Index);
-			}
-		}
-	};
+		NotifyFromFacesDiscrepancy(SkeletalMeshDest, SrcFaceVert.Num(), DstFaceVert.Num());
+		return false;
+	}
 
-	//Create a map linking all similar Position of destination vertex index
-	TMap<FVector3f, TArray<uint32>> PositionToVertexIndexDest;
-	PositionToVertexIndexDest.Reserve(VertexNumberSrc);
-	for (int32 VertexIndex = 0; VertexIndex < VertexNumberDest; ++VertexIndex)
+	// build unused vertices data
+	const int32 NumUnusedVertices = PointNumberDest-DstFaceVert.Num(); 
+	TSet<uint32> UnusedSrcFaceVert;
+	if (NumUnusedVertices > 0)
 	{
-		const FVector3f& Position = ImportDataDest.Points[VertexIndex];
-		TArray<uint32>& VertexIndexArray = PositionToVertexIndexDest.FindOrAdd(Position);
-		VertexIndexArray.Add(VertexIndex);
+		UE_LOG(LogLODUtilities, Warning, TEXT("Alternate skinning import: %d vertices do not reference any triangle."), NumUnusedVertices);
+		GetUnusedVertices(SrcFaceVert, NumUnusedVertices, UnusedSrcFaceVert);
 	}
-
+	
 	//Create a map to remap source bone index to destination bone index
 	TMap<int32, int32> RemapBoneIndexSrcToDest;
 	FillRemapBoneIndexSrcToDest(ImportDataSrc, ImportDataDest, RemapBoneIndexSrcToDest);
 
 	//Map to get the vertex index source to a destination vertex match
 	TSortedMap<uint32, VertexMatchNameSpace::FVertexMatchResult> VertexIndexSrcToVertexIndexDestMatches;
-	VertexIndexSrcToVertexIndexDestMatches.Reserve(VertexNumberSrc);
+	VertexIndexSrcToVertexIndexDestMatches.Reserve(DstFaceVert.Num());
 	TArray<uint32> VertexIndexToMatchWithPositions;
 
-	TMultiMap<int32, int32> VertexToSrcWedgeMap;
-	TMultiMap<int32, int32> VertexToDestWedgeMap;
-
-	for (int32 WedgeIndex = 0; WedgeIndex < ImportDataSrc.Wedges.Num(); WedgeIndex++)
-	{
-		const SkeletalMeshImportData::FVertex& Wedge = ImportDataSrc.Wedges[WedgeIndex];
-		VertexToSrcWedgeMap.Add(Wedge.VertexIndex, WedgeIndex);
-	}
-	for (int32 WedgeIndex = 0; WedgeIndex < ImportDataDest.Wedges.Num(); WedgeIndex++)
-	{
-		const SkeletalMeshImportData::FVertex& Wedge = ImportDataDest.Wedges[WedgeIndex];
-		VertexToDestWedgeMap.Add(Wedge.VertexIndex, WedgeIndex);
-	}
-
-	TArray<int32> SrcWedgeIndexes;
-	TArray<int32> DestWedgeIndexes;
-
+	// Match all used source vertex with destination vertex
+	const FVertexMatcher Matcher(ImportDataSrc, ImportDataDest, DstFaceVert, VertexIndexToMatchWithPositions, VertexIndexSrcToVertexIndexDestMatches);
+	Matcher.TryMatching(SrcFaceVert);
 	
-	// Match all source vertex with destination vertex
-	for (int32 VertexIndexSrc = 0; VertexIndexSrc < PointNumberSrc; ++VertexIndexSrc)
+	// Match all unused source vertex with destination vertex
+	if (NumUnusedVertices > 0)
 	{
-		const FVector3f& PositionSrc = ImportDataSrc.Points[VertexIndexSrc];
+		const int32 PrevNumMatch = VertexIndexSrcToVertexIndexDestMatches.Num();
 		
-		TArray<int32> SimilarDestinationVertex;
-		FindSimilarPosition(PositionSrc, SimilarDestinationVertex, KINDA_SMALL_NUMBER);
+		Matcher.TryMatching(UnusedSrcFaceVert);
 
-		if (SimilarDestinationVertex.Num() == 0)
+		const int32 NumMatch = VertexIndexSrcToVertexIndexDestMatches.Num();
+		if (NumMatch == PrevNumMatch && NumMatch < PointNumberSrc)
 		{
-			//Match with UV projection
-			VertexIndexToMatchWithPositions.Add(VertexIndexSrc);
-		}
-		else
-		{
-			//We have a direct match
-			VertexMatchNameSpace::FVertexMatchResult& VertexMatchDest = VertexIndexSrcToVertexIndexDestMatches.Add(VertexIndexSrc);
-
-			SrcWedgeIndexes.Reset();
-			VertexToSrcWedgeMap.MultiFind(VertexIndexSrc, SrcWedgeIndexes);
-
-			if (SrcWedgeIndexes.Num() > 0 && SimilarDestinationVertex.Num() > 1)
-			{
-				//Check if we have a point that is perfectly matching (position, UV, material and vertex color). Because normals and tangent are on the triangles we do not test those.
-				for (int32 MatchDestinationIndex = 0; MatchDestinationIndex < SimilarDestinationVertex.Num(); ++MatchDestinationIndex)
-				{
-					int32 VertexIndexDest = SimilarDestinationVertex[MatchDestinationIndex];
-					DestWedgeIndexes.Reset();
-					VertexToDestWedgeMap.MultiFind(VertexIndexDest, DestWedgeIndexes);
-					
-					for (int32 IndexDest = 0; IndexDest < DestWedgeIndexes.Num(); ++IndexDest)
-					{
-						int32 DestWedgeIndex = DestWedgeIndexes[IndexDest];
-						const SkeletalMeshImportData::FVertex& WedgeDest = ImportDataDest.Wedges[DestWedgeIndex];
-						for (int32 IndexSrc = 0; IndexSrc < SrcWedgeIndexes.Num(); ++IndexSrc)
-						{
-							int32 SrcWedgeIndex = SrcWedgeIndexes[IndexSrc];
-							const SkeletalMeshImportData::FVertex& WedgeSrc = ImportDataSrc.Wedges[SrcWedgeIndex];
-							//Wedge == operator test: material, vertex color and UVs
-							if (WedgeDest == WedgeSrc)
-							{
-								VertexMatchDest.VertexIndexes.Add(SimilarDestinationVertex[MatchDestinationIndex]);
-								VertexMatchDest.Ratios.Add(1.0f);
-								break;
-							}
-						}
-						if (VertexMatchDest.VertexIndexes.Num() > 0)
-						{
-							break;
-						}
-					}
-				}
-			}
-			//If there is no direct match, simply put everything
-			if (VertexMatchDest.VertexIndexes.Num() == 0)
-			{
-				for (int32 MatchDestinationIndex = 0; MatchDestinationIndex < SimilarDestinationVertex.Num(); ++MatchDestinationIndex)
-				{
-					VertexMatchDest.VertexIndexes.Add(SimilarDestinationVertex[MatchDestinationIndex]);
-					VertexMatchDest.Ratios.Add(1.0f / SimilarDestinationVertex.Num());
-				}
-			}
+			// try using position matching otherwise
+			MatchVertexIndexUsingPosition(ImportDataDest, ImportDataSrc, VertexIndexSrcToVertexIndexDestMatches, UnusedSrcFaceVert.Array(), bNoMatchMsgDone);
 		}
 	}
 	
@@ -2620,7 +2843,7 @@ bool FLODUtilities::UpdateAlternateSkinWeights(
 	{
 		MatchVertexIndexUsingPosition(ImportDataDest, ImportDataSrc, VertexIndexSrcToVertexIndexDestMatches, VertexIndexToMatchWithPositions, bNoMatchMsgDone);
 		//Make sure each vertex index source has a match, warn the user in case there is no match
-		for (int32 VertexIndexSource = 0; VertexIndexSource < VertexNumberSrc; ++VertexIndexSource)
+		for (int32 VertexIndexSource = 0; VertexIndexSource < PointNumberSrc; ++VertexIndexSource)
 		{
 			if (!VertexIndexSrcToVertexIndexDestMatches.Contains(VertexIndexSource))
 			{
@@ -2640,7 +2863,7 @@ bool FLODUtilities::UpdateAlternateSkinWeights(
 	TSortedMap<uint32, VertexMatchNameSpace::FVertexMatchResult> VertexIndexDestToVertexIndexSrcMatches;
 	if(!bAllSourceVertexAreMatch || PointNumberDest != PointNumberSrc)
 	{
-		VertexIndexDestToVertexIndexSrcMatches.Reserve(VertexNumberDest);
+		VertexIndexDestToVertexIndexSrcMatches.Reserve(PointNumberDest);
 		TArray<uint32> VertexIndexToMatch;
 		VertexIndexToMatch.Reserve(PointNumberDest);
 		for (int32 VertexIndexDest = 0; VertexIndexDest < PointNumberDest; ++VertexIndexDest)
@@ -2717,7 +2940,7 @@ bool FLODUtilities::UpdateAlternateSkinWeights(
 			InfluenceKeyToInfluenceIndex.Add(Key, AlternateInfluenceIndex);
 		}
 
-		for (int32 VertexIndexDestination = 0; VertexIndexDestination < VertexNumberDest; ++VertexIndexDestination)
+		for (int32 VertexIndexDestination = 0; VertexIndexDestination < PointNumberDest; ++VertexIndexDestination)
 		{
 			//Skip if the vertex is already matched
 			if (DestinationVertexIndexMatched[VertexIndexDestination])
@@ -2771,6 +2994,13 @@ bool FLODUtilities::UpdateAlternateSkinWeights(
 			}
 		}
 	}
+	else if (NumUnusedVertices > 0)
+	{
+		const FUnusedInfluenceBuilder InfluenceBuilder(
+			VertexIndexSrcToVertexIndexDestMatches, SourceVertexIndexToAlternateInfluenceIndexMap, PointNumberSrc,
+			AlternateInfluences);
+		InfluenceBuilder.TryBuilding(UnusedSrcFaceVert);
+	}
 	
 	// Check if there are any points that have no weight assignments after all is said and done. Attempt to average across neighbors to reconstruct
 	// those weights. Currently this relies on these points being completely surrounded by at least two points that have weights assigned to them.
@@ -5047,4 +5277,4 @@ bool FLODUtilities::StripLODGeometry(USkeletalMesh* SkeletalMesh, const int32 LO
 
 #undef LOCTEXT_NAMESPACE // "LODUtilities"
 
-#endif //WITH_EDITOR
+#endif //WITH_EDITOR
\ No newline at end of file
