ASW - change log from Warning to Display
    - follow up to 37017160

#rnx

diff --git a/Engine/Source/Developer/SkeletalMeshUtilitiesCommon/Private/LODUtilities.cpp b/Engine/Source/Developer/SkeletalMeshUtilitiesCommon/Private/LODUtilities.cpp
--- a/Engine/Source/Developer/SkeletalMeshUtilitiesCommon/Private/LODUtilities.cpp
+++ b/Engine/Source/Developer/SkeletalMeshUtilitiesCommon/Private/LODUtilities.cpp
@@ -2805,7 +2805,7 @@ bool FLODUtilities::UpdateAlternateSkinWeights(
 	TSet<uint32> UnusedSrcFaceVert;
 	if (NumUnusedVertices > 0)
 	{
-		UE_LOG(LogLODUtilities, Warning, TEXT("Alternate skinning import: %d vertices do not reference any triangle."), NumUnusedVertices);
+		UE_LOG(LogLODUtilities, Display, TEXT("Alternate skinning import: %d vertices do not reference any triangle."), NumUnusedVertices);
 		GetUnusedVertices(SrcFaceVert, NumUnusedVertices, UnusedSrcFaceVert);
 	}
 	
