UE: Lower verbosity on warning as it was determined to be benign

diff --git a/Engine/Source/Developer/SkeletalMeshUtilitiesCommon/Private/LODUtilities.cpp b/Engine/Source/Developer/SkeletalMeshUtilitiesCommon/Private/LODUtilities.cpp
--- a/Engine/Source/Developer/SkeletalMeshUtilitiesCommon/Private/LODUtilities.cpp
+++ b/Engine/Source/Developer/SkeletalMeshUtilitiesCommon/Private/LODUtilities.cpp
@@ -1170,7 +1170,7 @@ void ProjectTargetOnBase(const TArray<FSoftSkinVertex>& BaseVertices, const TArr
 		{
 			if(!bNoUVsMsgDisplayed)
 			{
-				UE_LOG(LogLODUtilities, Warning, TEXT("SkeletalMesh [%s] Remap morph target: Cannot remap morph target because source UVs are missings."), DebugContext ? DebugContext : TEXT("Unknown Source"));
+				UE_LOG(LogLODUtilities, Display, TEXT("SkeletalMesh [%s] Remap morph target: Cannot remap morph target because source UVs are missings."), DebugContext ? DebugContext : TEXT("Unknown Source"));
 				bNoUVsMsgDisplayed = true;
 			}
 			continue;
