LOD Utilities: Ensure we don't crash if the alt skin weights refer to a missing bone.

#rnx

diff --git a/Engine/Source/Developer/SkeletalMeshUtilitiesCommon/Private/LODUtilities.cpp b/Engine/Source/Developer/SkeletalMeshUtilitiesCommon/Private/LODUtilities.cpp
--- a/Engine/Source/Developer/SkeletalMeshUtilitiesCommon/Private/LODUtilities.cpp
+++ b/Engine/Source/Developer/SkeletalMeshUtilitiesCommon/Private/LODUtilities.cpp
@@ -6,6 +6,7 @@
 
 #if WITH_EDITOR
 
+
 #include "Algo/Accumulate.h"
 #include "Animation/MorphTarget.h"
 #include "Animation/SkinWeightProfile.h"
@@ -2886,9 +2887,9 @@ bool FLODUtilities::UpdateAlternateSkinWeights(
 		int32 VertexIndexSource = InfluenceSrc.VertexIndex;
 		int32 BoneIndexSource = InfluenceSrc.BoneIndex;
 		//We need to remap the source bone index to have the matching target bone index
-		int32 BoneIndexDest = RemapBoneIndexSrcToDest[BoneIndexSource];
-		if (BoneIndexDest != INDEX_NONE)
+		if (const int32* BoneIndexDstPtr = RemapBoneIndexSrcToDest.Find(BoneIndexSource))
 		{
+			const int32 BoneIndexDest = *BoneIndexDstPtr;
 			//Find the match destination vertex index
 			VertexMatchNameSpace::FVertexMatchResult* SourceVertexMatch = VertexIndexSrcToVertexIndexDestMatches.Find(VertexIndexSource);
 			if (SourceVertexMatch == nullptr || SourceVertexMatch->VertexIndexes.Num() <= 0)
@@ -5304,4 +5305,4 @@ bool FLODUtilities::StripLODGeometry(USkeletalMesh* SkeletalMesh, const int32 LO
 
 #undef LOCTEXT_NAMESPACE // "LODUtilities"
 
-#endif //WITH_EDITOR
\ No newline at end of file
+#endif //WITH_EDITOR
