Skeletal Mesh: When transfering alt skin weights from LOD to LOD during render mesh generation, ignore unconnected points completely so they don't affect connected points.

- Update the skelmesh DDC key to regenerate all skeletal meshes, since this information is not encoded in the key.

#rnx

diff --git a/Engine/Source/Developer/SkeletalMeshUtilitiesCommon/Private/LODUtilities.cpp b/Engine/Source/Developer/SkeletalMeshUtilitiesCommon/Private/LODUtilities.cpp
--- a/Engine/Source/Developer/SkeletalMeshUtilitiesCommon/Private/LODUtilities.cpp
+++ b/Engine/Source/Developer/SkeletalMeshUtilitiesCommon/Private/LODUtilities.cpp
@@ -2802,12 +2802,14 @@ bool FLODUtilities::UpdateAlternateSkinWeights(
 	}
 
 	// build unused vertices data
-	const int32 NumUnusedVertices = PointNumberDest-DstFaceVert.Num(); 
-	TSet<uint32> UnusedSrcFaceVert;
+	const int32 NumUsedVertices = DstFaceVert.Num();
+	const int32 NumUnusedVertices = PointNumberDest - NumUsedVertices; 
+	TSet<uint32> UnusedSrcFaceVert, UnusedDstFaceVert;
 	if (NumUnusedVertices > 0)
 	{
 		UE_LOG(LogLODUtilities, Display, TEXT("Alternate skinning import: %d vertices do not reference any triangle."), NumUnusedVertices);
 		GetUnusedVertices(SrcFaceVert, NumUnusedVertices, UnusedSrcFaceVert);
+		GetUnusedVertices(DstFaceVert, NumUnusedVertices, UnusedDstFaceVert);
 	}
 	
 	//Create a map to remap source bone index to destination bone index
@@ -2823,28 +2825,13 @@ bool FLODUtilities::UpdateAlternateSkinWeights(
 	const FVertexMatcher Matcher(ImportDataSrc, ImportDataDest, DstFaceVert, VertexIndexToMatchWithPositions, VertexIndexSrcToVertexIndexDestMatches);
 	Matcher.TryMatching(SrcFaceVert);
 	
-	// Match all unused source vertex with destination vertex
-	if (NumUnusedVertices > 0)
-	{
-		const int32 PrevNumMatch = VertexIndexSrcToVertexIndexDestMatches.Num();
-		
-		Matcher.TryMatching(UnusedSrcFaceVert);
-
-		const int32 NumMatch = VertexIndexSrcToVertexIndexDestMatches.Num();
-		if (NumMatch == PrevNumMatch && NumMatch < PointNumberSrc)
-		{
-			// try using position matching otherwise
-			MatchVertexIndexUsingPosition(ImportDataDest, ImportDataSrc, VertexIndexSrcToVertexIndexDestMatches, UnusedSrcFaceVert.Array(), bNoMatchMsgDone);
-		}
-	}
-	
 	//Find a match for all unmatched source vertex, unmatched vertex happen when the geometry is different between source and destination mesh
-	bool bAllSourceVertexAreMatch = VertexIndexToMatchWithPositions.Num() <= 0 && VertexIndexSrcToVertexIndexDestMatches.Num() == PointNumberSrc;
+	bool bAllSourceVertexAreMatch = VertexIndexToMatchWithPositions.Num() <= 0 && VertexIndexSrcToVertexIndexDestMatches.Num() == NumUsedVertices;
 	if (!bAllSourceVertexAreMatch)
 	{
 		MatchVertexIndexUsingPosition(ImportDataDest, ImportDataSrc, VertexIndexSrcToVertexIndexDestMatches, VertexIndexToMatchWithPositions, bNoMatchMsgDone);
-		//Make sure each vertex index source has a match, warn the user in case there is no match
-		for (int32 VertexIndexSource = 0; VertexIndexSource < PointNumberSrc; ++VertexIndexSource)
+		//Make sure each used vertex index source has a match, warn the user in case there is no match
+		for (int32 VertexIndexSource: SrcFaceVert)
 		{
 			if (!VertexIndexSrcToVertexIndexDestMatches.Contains(VertexIndexSource))
 			{
@@ -2856,21 +2843,16 @@ bool FLODUtilities::UpdateAlternateSkinWeights(
 				}
 			}
 		}
-		bAllSourceVertexAreMatch = VertexIndexSrcToVertexIndexDestMatches.Num() == PointNumberSrc;
+		bAllSourceVertexAreMatch = VertexIndexSrcToVertexIndexDestMatches.Num() == NumUsedVertices;
 	}
 	
 	
 	//Find the Destination to source match, to make sure all extra destination vertex get weighted properly in the alternate influences
 	TSortedMap<uint32, VertexMatchNameSpace::FVertexMatchResult> VertexIndexDestToVertexIndexSrcMatches;
-	if(!bAllSourceVertexAreMatch || PointNumberDest != PointNumberSrc)
+	if(!bAllSourceVertexAreMatch)
 	{
-		VertexIndexDestToVertexIndexSrcMatches.Reserve(PointNumberDest);
-		TArray<uint32> VertexIndexToMatch;
-		VertexIndexToMatch.Reserve(PointNumberDest);
-		for (int32 VertexIndexDest = 0; VertexIndexDest < PointNumberDest; ++VertexIndexDest)
-		{
-			VertexIndexToMatch.Add(VertexIndexDest);
-		}
+		VertexIndexDestToVertexIndexSrcMatches.Reserve(NumUsedVertices);
+		const TArray<uint32> VertexIndexToMatch = DstFaceVert.Array();
 		MatchVertexIndexUsingPosition(ImportDataSrc, ImportDataDest, VertexIndexDestToVertexIndexSrcMatches, VertexIndexToMatch, bNoMatchMsgDone);
 	}
 
@@ -2924,11 +2906,8 @@ bool FLODUtilities::UpdateAlternateSkinWeights(
 	//In case the source geometry was not matching the destination we have to add influence for each extra destination vertex index
 	if (VertexIndexDestToVertexIndexSrcMatches.Num() > 0)
 	{
-		TArray<bool> DestinationVertexIndexMatched;
-		DestinationVertexIndexMatched.AddZeroed(PointNumberDest);
-
-		int32 InfluenceNumberDest = ImportDataDest.Influences.Num();
-		int32 AlternateInfluenceNumber = AlternateInfluences.Num();
+		TBitArray<> DestinationVertexIndexMatched(false, PointNumberDest);
+		const int32 AlternateInfluenceNumber = AlternateInfluences.Num();
 		
 		//We want to avoid making duplicate so we use a map where the key is the boneindex mix with the destination vertex index
 		TMap<uint64, int32> InfluenceKeyToInfluenceIndex;
@@ -2995,13 +2974,6 @@ bool FLODUtilities::UpdateAlternateSkinWeights(
 			}
 		}
 	}
-	else if (NumUnusedVertices > 0)
-	{
-		const FUnusedInfluenceBuilder InfluenceBuilder(
-			VertexIndexSrcToVertexIndexDestMatches, SourceVertexIndexToAlternateInfluenceIndexMap, PointNumberSrc,
-			AlternateInfluences);
-		InfluenceBuilder.TryBuilding(UnusedSrcFaceVert);
-	}
 	
 	// Check if there are any points that have no weight assignments after all is said and done. Attempt to average across neighbors to reconstruct
 	// those weights. Currently this relies on these points being completely surrounded by at least two points that have weights assigned to them.
@@ -3010,6 +2982,13 @@ bool FLODUtilities::UpdateAlternateSkinWeights(
 	{
 		UnassignedVertexBitmask[Wedge.VertexIndex] = true;
 	}
+
+	// mark unused vertices as assigned to basically skip them
+	for (uint32 UnusedDstVertexIndex: UnusedDstFaceVert)
+	{
+		UnassignedVertexBitmask[UnusedDstVertexIndex] = false;
+	}
+
 	for (const SkeletalMeshImportData::FRawBoneInfluence& Influence: AlternateInfluences)
 	{
 		UnassignedVertexBitmask[Influence.VertexIndex] = false;
diff --git a/Engine/Source/Runtime/Core/Private/UObject/FortniteReleaseBranchCustomObjectVersion.cpp b/Engine/Source/Runtime/Core/Private/UObject/FortniteReleaseBranchCustomObjectVersion.cpp
--- a/Engine/Source/Runtime/Core/Private/UObject/FortniteReleaseBranchCustomObjectVersion.cpp
+++ b/Engine/Source/Runtime/Core/Private/UObject/FortniteReleaseBranchCustomObjectVersion.cpp
@@ -7,7 +7,7 @@ TMap<FGuid, FGuid> FFortniteReleaseBranchCustomObjectVersion::GetSystemGuids()
 	const FDevSystemGuids& DevGuids = FDevSystemGuids::Get();
 
 	SystemGuids.Add(DevGuids.Niagara_LatestScriptCompileVersion, FGuid("D35A655CD3FDA4408724D77E1316E35B"));
-	SystemGuids.Add(DevGuids.SkeletalMeshDerivedDataVersion, FGuid("5705956EC7134274A5A1BD3FC74DB3A9"));
+	SystemGuids.Add(DevGuids.SkeletalMeshDerivedDataVersion, FGuid("5DDF529BE4164CF6860F00156570398A"));
 	SystemGuids.Add(DevGuids.MaterialTranslationDDCVersion, FGuid("888A502A66F6401580B27451F211C46B"));
 	SystemGuids.Add(DevGuids.GLOBALSHADERMAP_DERIVEDDATA_VER, FGuid("0A0A1BF5D2A94105BB395EF0DA7CAEB9"));
 	SystemGuids.Add(DevGuids.MATERIALSHADERMAP_DERIVEDDATA_VER, FGuid("C0E1F59386894C9D855B85178B2553CB"));
