Skeletal Mesh: Downgrade warning on unmapped vertices to avoid cook failures.

#rnx

diff --git a/Engine/Source/Developer/SkeletalMeshUtilitiesCommon/Private/LODUtilities.cpp b/Engine/Source/Developer/SkeletalMeshUtilitiesCommon/Private/LODUtilities.cpp
--- a/Engine/Source/Developer/SkeletalMeshUtilitiesCommon/Private/LODUtilities.cpp
+++ b/Engine/Source/Developer/SkeletalMeshUtilitiesCommon/Private/LODUtilities.cpp
@@ -2424,7 +2424,7 @@ static void InfillMissingVertexSkinWeights(
 
 		if (!bPointRemoved)
 		{
-			UE_LOG(LogLODUtilities, Warning, TEXT("Alternate skinning import: Some alternate skinning vertices could not be mapped onto the base mesh. These will be skinned to the root bone."));
+			UE_LOG(LogLODUtilities, Display, TEXT("Alternate skinning import: Some alternate skinning vertices could not be mapped onto the base mesh. These will be skinned to the root bone."));
 			break;
 		}
 	}
