Skeletal Mesh: Move all morph object updates out of parallel-for worker threads and into the skinned mesh async compile thread to avoid threading issues with multiple threads touching the internal morph object data.

#rnx

diff --git a/Engine/Source/Developer/SkeletalMeshUtilitiesCommon/Private/LODUtilities.cpp b/Engine/Source/Developer/SkeletalMeshUtilitiesCommon/Private/LODUtilities.cpp
--- a/Engine/Source/Developer/SkeletalMeshUtilitiesCommon/Private/LODUtilities.cpp
+++ b/Engine/Source/Developer/SkeletalMeshUtilitiesCommon/Private/LODUtilities.cpp
@@ -1633,7 +1633,14 @@ void FLODUtilities::ApplyMorphTargetsToLOD(USkeletalMesh* SkeletalMesh, int32 So
 	CreateLODMorphTarget(SkeletalMesh, InlineReductionDataParameter, SourceLOD, DestinationLOD, PerMorphTargetBaseIndexToMorphTargetDelta, BaseMorphIndexToTargetIndexList, TargetVertices, TargetMatchData);
 }
 
-void FLODUtilities::SimplifySkeletalMeshLOD( USkeletalMesh* SkeletalMesh, int32 DesiredLOD, const ITargetPlatform* TargetPlatform, bool bRestoreClothing /*= false*/, FThreadSafeBool* OutNeedsPackageDirtied/*= nullptr*/)
+void FLODUtilities::SimplifySkeletalMeshLOD(
+	USkeletalMesh* SkeletalMesh,
+	int32 DesiredLOD,
+	const ITargetPlatform* TargetPlatform,
+	bool bRestoreClothing, 
+	TArray<FMorphTargetUpdateRequest>& OutMorphRequests,
+	std::atomic<bool>* OutNeedsPackageDirtied
+	)
 {
 	TRACE_CPUPROFILER_EVENT_SCOPE(FLODUtilities::SimplifySkeletalMeshLOD);
 
@@ -1763,11 +1770,11 @@ void FLODUtilities::SimplifySkeletalMeshLOD( USkeletalMesh* SkeletalMesh, int32
 			//Apply morph to the new LOD. Force it if we reduce the base LOD, base LOD must apply the morph target
 			if (ReductionSettings.bRemapMorphTargets)
 			{
-				ApplyMorphTargetsToLOD(SkeletalMesh, ReductionSettings.BaseLOD, DesiredLOD, InlineReductionDataParameter);
+				OutMorphRequests.Emplace(SkeletalMesh, ReductionSettings.BaseLOD, DesiredLOD, MoveTemp(InlineReductionDataParameter));
 			}
 			else
 			{
-				ClearGeneratedMorphTarget(SkeletalMesh, DesiredLOD);
+				OutMorphRequests.Emplace(SkeletalMesh, DesiredLOD);
 			}
 
 			// Update the vertex attribute information in the LOD info.
@@ -1812,7 +1819,7 @@ void FLODUtilities::SimplifySkeletalMeshLOD( USkeletalMesh* SkeletalMesh, int32
 	}
 }
 
-void FLODUtilities::SimplifySkeletalMeshLOD(FSkeletalMeshUpdateContext& UpdateContext, int32 DesiredLOD, const ITargetPlatform* TargetPlatform, bool bRestoreClothing /*= false*/, FThreadSafeBool* OutNeedsPackageDirtied/*= nullptr*/)
+void FLODUtilities::SimplifySkeletalMeshLOD(FSkeletalMeshUpdateContext& UpdateContext, int32 DesiredLOD, const ITargetPlatform* TargetPlatform, bool bRestoreClothing /*= false*/, std::atomic<bool>* OutNeedsPackageDirtied/*= nullptr*/)
 {
 	USkeletalMesh* SkeletalMesh = UpdateContext.SkeletalMesh;
 	IMeshReductionModule& ReductionModule = FModuleManager::Get().LoadModuleChecked<IMeshReductionModule>("MeshReductionInterface");
@@ -1820,7 +1827,13 @@ void FLODUtilities::SimplifySkeletalMeshLOD(FSkeletalMeshUpdateContext& UpdateCo
 
 	if (MeshReduction && MeshReduction->IsSupported() && SkeletalMesh)
 	{
-		SimplifySkeletalMeshLOD(SkeletalMesh, DesiredLOD, TargetPlatform, bRestoreClothing, OutNeedsPackageDirtied);
+		TArray<FMorphTargetUpdateRequest> MorphTargetUpdateRequests;
+		SimplifySkeletalMeshLOD(SkeletalMesh, DesiredLOD, TargetPlatform, bRestoreClothing, MorphTargetUpdateRequests, OutNeedsPackageDirtied);
+
+		for (FMorphTargetUpdateRequest& Request: MorphTargetUpdateRequests)
+		{
+			Request.Run();
+		}
 		
 		if (UpdateContext.OnLODChanged.IsBound())
 		{
@@ -3480,7 +3493,7 @@ void FLODUtilities::RegenerateDependentLODs(USkeletalMesh* SkeletalMesh, int32 L
 			}
 
 			//Reduce all dependent LODs
-			FThreadSafeBool bNeedsPackageDirtied(false);
+			std::atomic<bool> bNeedsPackageDirtied{false};
 			
 			//Adjust the InlineReductionCacheDatas before simplifying dependent LODs
 			if (SkeletalMesh->GetImportedModel()->InlineReductionCacheDatas.Num() < LODNumber)
@@ -3493,18 +3506,34 @@ void FLODUtilities::RegenerateDependentLODs(USkeletalMesh* SkeletalMesh, int32 L
 				SkeletalMesh->GetImportedModel()->InlineReductionCacheDatas.SetNum(LODNumber);
 			}
 
+			struct FSimplifyThreadContext
+			{
+				TArray<FMorphTargetUpdateRequest> MorphTargetUpdateRequests;
+			};
+
 			// Reduce LODs in parallel (reduction is multithread safe)
 			const bool bHasAccessToLockedProperties = !FSkinnedAssetAsyncBuildScope::ShouldWaitOnLockedProperties(SkeletalMesh);
-			ParallelFor(DependentLODs.Num(), [&DependentLODs, &SkeletalMesh, &bNeedsPackageDirtied, bHasAccessToLockedProperties, &TargetPlatform](int32 IterationIndex)
+			TArray<FSimplifyThreadContext> SimplifyThreadContexts;
+			ParallelForWithTaskContext(SimplifyThreadContexts, DependentLODs.Num(),
+				[&DependentLODs, &SkeletalMesh, &bNeedsPackageDirtied, bHasAccessToLockedProperties, &TargetPlatform](FSimplifyThreadContext& Context, int32 IterationIndex)
 			{
 				TUniquePtr<FSkinnedAssetAsyncBuildScope> AsyncBuildScope(bHasAccessToLockedProperties ? MakeUnique<FSkinnedAssetAsyncBuildScope>(SkeletalMesh) : nullptr);
 
 				check(DependentLODs.IsValidIndex(IterationIndex));
 				int32 DependentLODIndex = DependentLODs[IterationIndex];
 				check(SkeletalMesh->GetLODInfo(DependentLODIndex)); //We cannot add a LOD when reducing with multi thread, so check we already have one
-				FLODUtilities::SimplifySkeletalMeshLOD(SkeletalMesh, DependentLODIndex, TargetPlatform, false, &bNeedsPackageDirtied);
+				FLODUtilities::SimplifySkeletalMeshLOD(SkeletalMesh, DependentLODIndex, TargetPlatform, false, Context.MorphTargetUpdateRequests, &bNeedsPackageDirtied);
 			});
 
+			// Perform the morph target updates that we cannot do in a separate thread.
+			for (FSimplifyThreadContext& Context: SimplifyThreadContexts)
+			{
+				for (FMorphTargetUpdateRequest& Request: Context.MorphTargetUpdateRequests)
+				{
+					Request.Run();
+				}
+			}
+
 			if (bNeedsPackageDirtied && IsInGameThread())
 			{
 				SkeletalMesh->MarkPackageDirty();
diff --git a/Engine/Source/Developer/SkeletalMeshUtilitiesCommon/Public/LODUtilities.h b/Engine/Source/Developer/SkeletalMeshUtilitiesCommon/Public/LODUtilities.h
--- a/Engine/Source/Developer/SkeletalMeshUtilitiesCommon/Public/LODUtilities.h
+++ b/Engine/Source/Developer/SkeletalMeshUtilitiesCommon/Public/LODUtilities.h
@@ -100,7 +100,7 @@ public:
 	* @param UpdateContext - The skeletal mesh and actor components to operate on.
 	* @param DesiredLOD - The LOD to simplify
 	*/
-	static void SimplifySkeletalMeshLOD(FSkeletalMeshUpdateContext& UpdateContext, int32 DesiredLOD, const class ITargetPlatform* TargetPlatform, bool bRestoreClothing = false, class FThreadSafeBool* OutNeedsPackageDirtied = nullptr);
+	static void SimplifySkeletalMeshLOD(FSkeletalMeshUpdateContext& UpdateContext, int32 DesiredLOD, const class ITargetPlatform* TargetPlatform, bool bRestoreClothing = false, std::atomic<bool>* OutNeedsPackageDirtied = nullptr);
 
 	/**
 	*	Restore the LOD imported model to the last imported data. Call this function if you want to remove the reduce on the base LOD
@@ -326,6 +326,39 @@ private:
 		int32 BoneInfluenceLimit,
 		const class ITargetPlatform* TargetPlatform);
 
+	struct FMorphTargetUpdateRequest
+	{
+		FMorphTargetUpdateRequest() = default;
+		FMorphTargetUpdateRequest(USkeletalMesh* InSkeletalMesh, const int32 InSourceLOD, const int32 InDestinationLOD, FInlineReductionDataParameter&& InReductionData) :
+			SkeletalMesh(InSkeletalMesh),
+			SourceLOD(InSourceLOD),
+			DestinationLOD(InDestinationLOD),
+			ReductionData(MoveTemp(InReductionData))
+		{}
+
+		FMorphTargetUpdateRequest(USkeletalMesh* InSkeletalMesh, const int32 InDestinationLOD) :
+			SkeletalMesh(InSkeletalMesh),
+			DestinationLOD(InDestinationLOD)
+		{}
+		
+		USkeletalMesh* SkeletalMesh = nullptr; 
+		int32 SourceLOD = INDEX_NONE;
+		int32 DestinationLOD = INDEX_NONE;
+		FInlineReductionDataParameter ReductionData;
+
+		void Run() const
+		{
+			if (SourceLOD != INDEX_NONE)
+			{
+				ApplyMorphTargetsToLOD(SkeletalMesh, SourceLOD, DestinationLOD, ReductionData);
+			}
+			else
+			{
+				ClearGeneratedMorphTarget(SkeletalMesh, DestinationLOD);
+			}
+		}
+	};
+
 	/**
 	 *	Simplifies the static mesh based upon various user settings for DesiredLOD
 	 *  This is private function that gets called by SimplifySkeletalMesh
@@ -333,7 +366,14 @@ private:
 	 * @param SkeletalMesh - The skeletal mesh and actor components to operate on.
 	 * @param DesiredLOD - Desired LOD
 	 */
-	static void SimplifySkeletalMeshLOD(USkeletalMesh* SkeletalMesh, int32 DesiredLOD, const class ITargetPlatform* TargetPlatform, bool bRestoreClothing = false, class FThreadSafeBool* OutNeedsPackageDirtied = nullptr);
+	static void SimplifySkeletalMeshLOD(
+		USkeletalMesh* SkeletalMesh, 
+		int32 DesiredLOD, 
+		const ITargetPlatform* TargetPlatform,
+		bool bRestoreClothing, 
+		TArray<FMorphTargetUpdateRequest>& OutMorphRequests,
+		std::atomic<bool>* OutNeedsPackageDirtied = nullptr
+		);
 
 	/**
 	*  Remap the morph targets of the base LOD onto the desired LOD.
diff --git a/Engine/Source/Runtime/Engine/Public/Rendering/SkeletalMeshLODModel.h b/Engine/Source/Runtime/Engine/Public/Rendering/SkeletalMeshLODModel.h
--- a/Engine/Source/Runtime/Engine/Public/Rendering/SkeletalMeshLODModel.h
+++ b/Engine/Source/Runtime/Engine/Public/Rendering/SkeletalMeshLODModel.h
@@ -1,24 +1,20 @@
 // Copyright Epic Games, Inc. All Rights Reserved.
 
 #pragma once
-#include "Animation/SkeletalMeshVertexAttribute.h"
 
 #if WITH_EDITOR
 
-#include "CoreMinimal.h"
-#include "Engine/EngineTypes.h"
-#include "PackedNormal.h"
-#include "Components.h"
-#include "GPUSkinPublicDefs.h"
-#include "BoneIndices.h"
-#include "Serialization/BulkData.h"
-#include "SkeletalMeshTypes.h"
-#include "Rendering/SkeletalMeshLODImporterData.h"
+#include "Animation/SkeletalMeshVertexAttribute.h"
 #include "Animation/SkinWeightProfile.h"
-#include "CoreTypes.h"
 #include "Async/Mutex.h"
 #include "Async/UniqueLock.h"
+#include "BoneIndices.h"
+#include "Components.h"
+#include "GPUSkinPublicDefs.h"
 #include "HAL/CriticalSection.h"
+#include "Rendering/SkeletalMeshLODImporterData.h"
+#include "Serialization/BulkData.h"
+#include "SkeletalMeshTypes.h"
 
 //
 //	FSoftSkinVertex
@@ -421,6 +417,29 @@ public:
 		}
 	}
 
+	FSkeletalMeshLODModel(FSkeletalMeshLODModel&& InOther)
+	{
+		Sections = MoveTemp(InOther.Sections);
+		UserSectionsData = MoveTemp(InOther.UserSectionsData);
+		NumVertices = InOther.NumVertices;
+		NumTexCoords = InOther.NumTexCoords;
+		IndexBuffer = MoveTemp(InOther.IndexBuffer);
+		ActiveBoneIndices = MoveTemp(InOther.ActiveBoneIndices);
+		RequiredBones = MoveTemp(InOther.RequiredBones);
+		SkinWeightProfiles = MoveTemp(InOther.SkinWeightProfiles);
+		VertexAttributes = MoveTemp(InOther.VertexAttributes);
+		MeshToImportVertexMap = MoveTemp(InOther.MeshToImportVertexMap);
+		MaxImportVertex = InOther.MaxImportVertex;
+		RawSkeletalMeshBulkDataID = MoveTemp(InOther.RawSkeletalMeshBulkDataID);
+		bIsBuildDataAvailable = InOther.bIsBuildDataAvailable;
+		bIsRawSkeletalMeshBulkDataEmpty = InOther.bIsRawSkeletalMeshBulkDataEmpty;
+		ImportedMeshInfos = MoveTemp(InOther.ImportedMeshInfos);
+		RawPointIndices2 = MoveTemp(InOther.RawPointIndices2);
+		BulkDataReadMutex = InOther.BulkDataReadMutex;
+		InOther.BulkDataReadMutex = nullptr;
+		BuildStringID = InOther.BuildStringID;
+	}
+
 	/*Empty the skeletal mesh LOD model. Empty copy a default constructed FSkeletalMeshLODModel but will not copy the BulkDataReadMutex which will be the same after*/
 	void Empty()
 	{
