Skeletal Mesh: Fix the way we compute unused vertices for alt skin processing.

#rnx

diff --git a/Engine/Source/Developer/SkeletalMeshUtilitiesCommon/Private/LODUtilities.cpp b/Engine/Source/Developer/SkeletalMeshUtilitiesCommon/Private/LODUtilities.cpp
--- a/Engine/Source/Developer/SkeletalMeshUtilitiesCommon/Private/LODUtilities.cpp
+++ b/Engine/Source/Developer/SkeletalMeshUtilitiesCommon/Private/LODUtilities.cpp
@@ -2443,6 +2443,7 @@ static void InfillMissingVertexSkinWeights(
 	}
 }
 
+
 namespace AlternateSkinWeightsLocal
 {
 
@@ -2481,21 +2482,17 @@ void GetVerticesFromFaces(const FSkeletalMeshImportData& InImportData, TSet<uint
 }
 
 // note that this algo supposes that InUsed is sorted
-void GetUnusedVertices(const TSet<uint32>& InUsed, const int32 InExpectedUnusedNum, TSet<uint32>& OutUnused)
+void GetUnusedVertices(const TSet<uint32>& InUsed, const int32 InNumPoints, TSet<uint32>& OutUnused)
 {
-	OutUnused.Empty(InExpectedUnusedNum);
+	OutUnused.Empty(InNumPoints - InUsed.Num());
 
-	int32 UnusedVertIndex = 0;
-	for (int32 UsedVertIndex: InUsed)
+	for (int32 Index = 0; Index < InNumPoints; ++Index)
 	{
-		while (UnusedVertIndex < UsedVertIndex)
+		if (!InUsed.Contains(Index))
 		{
-			OutUnused.Add(UnusedVertIndex++);
+			OutUnused.Add(Index);
 		}
-		UnusedVertIndex = UsedVertIndex+1;
 	}
-	
-	ensure(OutUnused.Num() == InExpectedUnusedNum);
 }
 	
 class FVertexMatcher
@@ -2821,8 +2818,8 @@ bool FLODUtilities::UpdateAlternateSkinWeights(
 	if (NumUnusedVertices > 0)
 	{
 		UE_LOG(LogLODUtilities, Display, TEXT("Alternate skinning import: %d vertices do not reference any triangle."), NumUnusedVertices);
-		GetUnusedVertices(SrcFaceVert, NumUnusedVertices, UnusedSrcFaceVert);
-		GetUnusedVertices(DstFaceVert, NumUnusedVertices, UnusedDstFaceVert);
+		GetUnusedVertices(SrcFaceVert, ImportDataSrc.Points.Num(), UnusedSrcFaceVert);
+		GetUnusedVertices(DstFaceVert, ImportDataDest.Points.Num(), UnusedDstFaceVert);
 	}
 	
 	//Create a map to remap source bone index to destination bone index
