Skeletal Mesh: Fix multiple issues with alt skin weight profile application during render mesh generation.
- First by observing that the import meshes being used are the same mesh but with different influences.
- Second by using the correct map when translating from render point indices to import data point indices when resolving alt skin weights and vertex attributes.

#rnx

diff --git a/Engine/Plugins/Experimental/SkeletalReduction/Source/Private/SkeletalMeshReductionPlugin.cpp b/Engine/Plugins/Experimental/SkeletalReduction/Source/Private/SkeletalMeshReductionPlugin.cpp
--- a/Engine/Plugins/Experimental/SkeletalReduction/Source/Private/SkeletalMeshReductionPlugin.cpp
+++ b/Engine/Plugins/Experimental/SkeletalReduction/Source/Private/SkeletalMeshReductionPlugin.cpp
@@ -1539,13 +1539,11 @@ void  FQuadricSkeletalMeshReduction::AddSourceModelInfluences( const FSkeletalMe
 						FBoneIndexType LocalBoneId = SrcRawSkinWeight.InfluenceBones[b];
 						uint16 Weight = SrcRawSkinWeight.InfluenceWeights[b];
 
-						checkSlow(LocalBoneId < BoneMap.Num());
-
-						// decode the bone weight
-						int32 BoneId = BoneMap[LocalBoneId];
-
-						if (Weight > 0)
+						if (Weight > 0 && ensure(BoneMap.IsValidIndex(LocalBoneId)))
 						{
+							// decode the bone weight
+							int32 BoneId = BoneMap[LocalBoneId];
+							
 							SkeletalMeshImportData::FRawBoneInfluence RawBoneInfluence = { (float)Weight / UE::AnimationCore::MaxRawBoneWeightFloat,  VIdx, BoneId };
 
 							RawBoneInfluences.Add(RawBoneInfluence);
@@ -1636,7 +1634,7 @@ void  FQuadricSkeletalMeshReduction::AddSourceModelInfluences( const FSkeletalMe
 			SkinWeights.AddUninitialized(LODModel.NumVertices);
 
 			// map the verts in the LOD model to the input order.
-			const TArray<int32>& ImportVertexMap = LODModel.MeshToImportVertexMap;
+			const TArray<uint32>& ImportVertexMap = LODModel.GetRawPointIndices();
 			check(ImportVertexMap.Num() == (int32)LODModel.NumVertices);
 			int32 MaxNumInfluences = 0;
 			for (int32 VertexIndex = 0; VertexIndex < (int32)LODModel.NumVertices; ++VertexIndex)
@@ -1786,7 +1784,7 @@ void FQuadricSkeletalMeshReduction::ConvertToFSkeletalMeshLODModel(
 
 		for (int32 Index = 0; Index < SkinnedMesh.NumVertices(); Index++)
 		{
-			const int32 OriginalIndex = NewModel.MeshToImportVertexMap[Index];
+			const int32 OriginalIndex = NewModel.GetRawPointIndices()[Index];
 			const auto& Vertex = SkinnedMesh.VertexBuffer[OriginalIndex];
 			for (int32 Component = 0; Component < VertexAttribute.ComponentCount; Component++)
 			{
diff --git a/Engine/Source/Developer/SkeletalMeshUtilitiesCommon/Private/LODUtilities.cpp b/Engine/Source/Developer/SkeletalMeshUtilitiesCommon/Private/LODUtilities.cpp
--- a/Engine/Source/Developer/SkeletalMeshUtilitiesCommon/Private/LODUtilities.cpp
+++ b/Engine/Source/Developer/SkeletalMeshUtilitiesCommon/Private/LODUtilities.cpp
@@ -2,15 +2,13 @@
 
 #include "LODUtilities.h"
 
-#include "SkeletalMeshAttributes.h"
-
 #if WITH_EDITOR
 
-
 #include "Algo/Accumulate.h"
 #include "Animation/MorphTarget.h"
 #include "Animation/SkinWeightProfile.h"
 #include "Async/ParallelFor.h"
+#include "Async/TaskGraphInterfaces.h"
 #include "BoneWeights.h"
 #include "ClothingAsset.h"
 #include "ComponentReregisterContext.h"
@@ -25,13 +23,12 @@
 #include "Engine/Texture2D.h"
 #include "Framework/Notifications/NotificationManager.h"
 #include "GenericQuadTree.h"
-#include "HAL/ThreadSafeBool.h"
 #include "IMeshReductionManagerModule.h"
 #include "ImageCore.h"
 #include "ImageCoreUtils.h"
 #include "ImportUtils/SkelImport.h"
 #include "Interfaces/ITargetPlatform.h"
-#include "Interfaces/ITargetPlatformManagerModule.h"
+#include "Logging/StructuredLog.h"
 #include "MeshUtilities.h"
 #include "MeshUtilitiesCommon.h"
 #include "Misc/FeedbackContext.h"
@@ -42,8 +39,8 @@
 #include "Rendering/SkeletalMeshLODImporterData.h"
 #include "Rendering/SkeletalMeshLODModel.h"
 #include "Rendering/SkeletalMeshModel.h"
+#include "SkeletalMeshAttributes.h"
 #include "Tasks/Task.h"
-#include "Async/TaskGraphInterfaces.h"
 #include "UObject/GarbageCollection.h"
 #include "UObject/UObjectIterator.h"
 #include "Widgets/Notifications/SNotificationList.h"
@@ -2292,463 +2289,17 @@ bool FLODUtilities::UpdateAlternateSkinWeights(
 }
 
 
-static void InfillMissingVertexSkinWeights(
-	const FSkeletalMeshImportData& InMesh,
-	const TArray<int32>& InMissingVertexes,
-	TArray<SkeletalMeshImportData::FRawBoneInfluence> &InOutInfluences
-	)
+static bool IsSameImportMeshStructure(const FSkeletalMeshImportData& MeshA, const FSkeletalMeshImportData& MeshB)
 {
-	// Create a map from influences to vertices.
-	TMultiMap<int32 /* Vertex */, int32 /* AlternateInfluencesIndex */> VertexToInfluenceMap;
-	VertexToInfluenceMap.Reserve(InOutInfluences.Num());
-	for (int32 InfluenceIndex = 0; InfluenceIndex < InOutInfluences.Num(); InfluenceIndex++)
-	{
-		VertexToInfluenceMap.Add(InOutInfluences[InfluenceIndex].VertexIndex, InfluenceIndex);
-	}
-
-	// Map from vertices to faces so we can track down neighbors. 
-	TMultiMap<int32 /* Vertex */, int32 /* Face */> VertexToFaceMap;
-	VertexToFaceMap.Reserve(InMesh.Wedges.Num() * 3);
-	for (int32 FaceIndex = 0; FaceIndex < InMesh.Faces.Num(); FaceIndex++)
-	{
-		for (int32 CornerIndex = 0; CornerIndex < 3; CornerIndex++)
-		{
-			const int32 WedgeIndex = InMesh.Faces[FaceIndex].WedgeIndex[CornerIndex];
-			VertexToFaceMap.Add(InMesh.Wedges[WedgeIndex].VertexIndex, FaceIndex);
-		}
-	}
-
-	// Compute neighborhood vertices for each unassigned point and assign them weights based on half the area of the triangle
-	// they belong to. We use this area for final weighting of influence contributions.
-	TArray<int32> Faces;
-	TMultiMap<int32 /* Vertex */, TPair<int32 /* NeighborVertex */, float /* NeighborArea */>> NeighborVertexesAndArea;
-
-	// We tally up the ratio of assigned influences for each unassigned vertex ranging from all neighbors have influences (1.0) to
-	// none have influences (0.0).
-	TArray<float> AssignedNeighborWeightRatio;
-	
-	for (int32 UnassignedVertex: InMissingVertexes)
-	{
-		Faces.Reset();
-		VertexToFaceMap.MultiFind(UnassignedVertex, Faces);
-		check(!Faces.IsEmpty());
-
-		TMap<int32, float> NeighborAndArea;
-
-		for (const int32 FaceIndex: Faces)
-		{
-			int32 VertexIndex[3];
-			
-			for (int32 CornerIndex = 0; CornerIndex < 3; CornerIndex++)
-			{
-				const int32 WedgeIndex  = InMesh.Faces[FaceIndex].WedgeIndex[CornerIndex];
-				VertexIndex[CornerIndex] = InMesh.Wedges[WedgeIndex].VertexIndex;
-			}
-
-			float Area = TriangleUtilities::ComputeTriangleArea(InMesh.Points[VertexIndex[0]], InMesh.Points[VertexIndex[1]], InMesh.Points[VertexIndex[2]]);
-			for (int32 CornerIndex = 0; CornerIndex < 3; CornerIndex++)
-			{
-				if (VertexIndex[CornerIndex] != UnassignedVertex)
-				{
-					NeighborAndArea.FindOrAdd(VertexIndex[CornerIndex], 0.0f) += Area / 2.0f;
-				}
-			}
-		}
-		int32 AssignedCount = 0;
-		for (const TPair<int32, float>& Item: NeighborAndArea)
-		{
-			NeighborVertexesAndArea.Add(UnassignedVertex, Item);
-
-			// Check if this neighbor has any influences computed already.
-			if (VertexToInfluenceMap.Contains(Item.Key))
-			{
-				AssignedCount++;
-			}
-		}
-		AssignedNeighborWeightRatio.Add(AssignedCount / static_cast<float>(NeighborAndArea.Num()));
-	}
-
-	// Order the unassigned vertices by how many counts are assigned. Start with the fully assigned ones first.
-	TSet<int32> UnassignedVertexQueue;
-	{
-		TArray<TPair<int32, float>> VertexWithRatio;
-		for (int32 Index = 0; Index < InMissingVertexes.Num(); Index++)
-		{
-			VertexWithRatio.Add({InMissingVertexes[Index], AssignedNeighborWeightRatio[Index]});
-		}
-		// Sort in descending order of ratio.
-		VertexWithRatio.Sort([](const TPair<int32, float>& A, const TPair<int32, float>& B) { return A.Value > B.Value; });
-
-		for (const TPair<int32, float>& Item: VertexWithRatio)
-		{
-			UnassignedVertexQueue.Add(Item.Key);
-		}
-	}
-
-	// Progressively try to resolve weights. If we end up with a set of points we can't resolve (e.g. islands, etc. we warn the user that some
-	// of their vertices will be mapped to root).
-	TArray<TPair<int32 /* NeighborVertex */, float /* NeighborArea */>> NeighborVertexAreaLookup;
-	TArray<int32> InfluenceIndexLookup;
-	while (!UnassignedVertexQueue.IsEmpty())
-	{
-		bool bPointRemoved = false;
-		for (int32 UnassignedVertex: UnassignedVertexQueue)
-		{
-			NeighborVertexAreaLookup.Reset();
-			NeighborVertexesAndArea.MultiFind(UnassignedVertex, NeighborVertexAreaLookup);
-
-			if (NeighborVertexAreaLookup.IsEmpty())
-			{
-				continue;
-			}
-
-			TMap</* BoneIndex */ int32, /* Weight */ float> VertexInfluences;
-			// Count up the influences that each neighbor has and weigh them by area.
-			for (const TPair<int32 /* NeighborVertex */, float /* NeighborArea */>& Item: NeighborVertexAreaLookup)
-			{
-				InfluenceIndexLookup.Reset();
-				VertexToInfluenceMap.MultiFind(Item.Key, InfluenceIndexLookup);
-
-				for (int32 InfluenceIndex: InfluenceIndexLookup)
-				{
-					const SkeletalMeshImportData::FRawBoneInfluence& Influence = InOutInfluences[InfluenceIndex];
-					VertexInfluences.FindOrAdd(Influence.BoneIndex, 0.0f) += Influence.Weight * Item.Value;
-				}
-			}
-
-			if (VertexInfluences.IsEmpty())
-			{
-				continue;
-			}
-
-			// Normalize the weights and add them to the list of alternate weights.
-			float SumWeight = Algo::Accumulate(VertexInfluences, 0.0f, [](float V, const TPair<int32, float>& Item) { return V + Item.Value; });
-
-			for (const TPair<int32, float>& Item: VertexInfluences)
-			{
-				VertexToInfluenceMap.Add(UnassignedVertex, InOutInfluences.Num());
-				InOutInfluences.Add({Item.Value / SumWeight, UnassignedVertex, Item.Key});
-			}
-
-			UnassignedVertexQueue.Remove(UnassignedVertex);
-			bPointRemoved = true;
-			break;
-		}
-
-		if (!bPointRemoved)
-		{
-			UE_LOG(LogLODUtilities, Display, TEXT("Alternate skinning import: Some alternate skinning vertices could not be mapped onto the base mesh. These will be skinned to the root bone."));
-			break;
-		}
-	}
+	// This equivalency testing above will suffice and is only for sanity, since both of the import data objects
+	// passed in should have been generated from the same FMeshDescription (see FSkeletalMeshImportData::GetMeshDescription)
+	return MeshA.Points.Num() == MeshB.Points.Num() && 
+	       MeshA.Wedges.Num() == MeshB.Wedges.Num() &&
+		   MeshA.Faces.Num() == MeshB.Faces.Num() &&
+	       MeshA.RefBonesBinary.Num() == MeshB.RefBonesBinary.Num();
 }
 
 
-namespace AlternateSkinWeightsLocal
-{
-
-void NotifyFromFacesDiscrepancy(const USkeletalMesh* InSkeletalMeshDest, const int32 InNumSrcFaceVerts, const int32 InNumDstFaceVerts)
-{
-	FFormatNamedArguments Args;
-	Args.Add(TEXT("SkeletalMeshName"), FText::FromString(InSkeletalMeshDest->GetName()));
-	Args.Add(TEXT("FacesPointNumberSrc"), InNumSrcFaceVerts);
-	Args.Add(TEXT("FacesPointNumberDest"), InNumDstFaceVerts);
-
-	FText Message = FText::Format(NSLOCTEXT("FLODUtilities_UpdateAlternateSkinWeights", "DifferentFaceVerticesNumber",
-		"Asset {SkeletalMeshName} failed to import skin weight profile the alternate skinning model has a different number of face vertex."
-		" Alternate face vertex count: {FacesPointNumberSrc}, LOD vertex count: {FacesPointNumberDest}"), Args);
-	
-	UE_LOG(LogLODUtilities, Warning, TEXT("%s"), *(Message.ToString()));
-}
-
-void GetVerticesFromFaces(const FSkeletalMeshImportData& InImportData, TSet<uint32>& OutFaceVertices)
-{
-	const TConstArrayView<SkeletalMeshImportData::FTriangle> Faces(InImportData.Faces);
-	const TConstArrayView<SkeletalMeshImportData::FVertex> Wedges(InImportData.Wedges);
-	
-	OutFaceVertices.Empty(Faces.Num());
-	
-	for (const SkeletalMeshImportData::FTriangle& Triangle: Faces)
-	{
-		OutFaceVertices.Add(Wedges[Triangle.WedgeIndex[0]].VertexIndex);
-		OutFaceVertices.Add(Wedges[Triangle.WedgeIndex[1]].VertexIndex);
-		OutFaceVertices.Add(Wedges[Triangle.WedgeIndex[2]].VertexIndex);
-	}
- 
-	OutFaceVertices.Sort([](uint32 VertIndex0, uint32 VertIndex1)
-	{
-		return VertIndex0 < VertIndex1;
-	});		
-}
-
-// note that this algo supposes that InUsed is sorted
-void GetUnusedVertices(const TSet<uint32>& InUsed, const int32 InNumPoints, TSet<uint32>& OutUnused)
-{
-	OutUnused.Empty(InNumPoints - InUsed.Num());
-
-	for (int32 Index = 0; Index < InNumPoints; ++Index)
-	{
-		if (!InUsed.Contains(Index))
-		{
-			OutUnused.Add(Index);
-		}
-	}
-}
-	
-class FVertexMatcher
-{
-public:
-	
-	FVertexMatcher(
-		const FSkeletalMeshImportData& InImportDataSrc,
-		const FSkeletalMeshImportData& InImportDataDst,
-		const TSet<uint32>& InDstVertices,
-		TArray<uint32>& InOutVertexIndexToMatchWithPositions,
-		TSortedMap<uint32, VertexMatchNameSpace::FVertexMatchResult>& InOutVertexIndexSrcToVertexIndexDestMatches)
-	: ImportDataSrc(InImportDataSrc)
-	, ImportDataDst(InImportDataDst)
-	, DstVertices(InDstVertices)
-	, VertexIndexToMatchWithPositions(InOutVertexIndexToMatchWithPositions)
-	, VertexIndexSrcToVertexIndexDestMatches(InOutVertexIndexSrcToVertexIndexDestMatches)
-	{
-		for (int32 WedgeIndex = 0; WedgeIndex < ImportDataSrc.Wedges.Num(); WedgeIndex++)
-		{
-			const SkeletalMeshImportData::FVertex& Wedge = ImportDataSrc.Wedges[WedgeIndex];
-			VertexToSrcWedgeMap.Add(Wedge.VertexIndex, WedgeIndex);
-		}
-		for (int32 WedgeIndex = 0; WedgeIndex < ImportDataDst.Wedges.Num(); WedgeIndex++)
-		{
-			const SkeletalMeshImportData::FVertex& Wedge = ImportDataDst.Wedges[WedgeIndex];
-			VertexToDstWedgeMap.Add(Wedge.VertexIndex, WedgeIndex);
-		}
-	}
-
-	void TryMatching(const TSet<uint32>& InSrcVertices) const
-	{
-		// Create a list of vertex Z/index pairs
-		TArray<FIndexAndZ> VertIndexAndZ;
-		VertIndexAndZ.Reserve(DstVertices.Num());
-		for (int32 VertexIndex: DstVertices)
-		{
-			new(VertIndexAndZ)FIndexAndZ(VertexIndex, ImportDataDst.Points[VertexIndex]);
-		}
-		// Project the vertices onto a diagonal, use that to quickly find vertices for
-		// closest point matching. 
-		VertIndexAndZ.Sort(FCompareIndexAndZ());
-		
-		auto FindSimilarPosition = [&VertIndexAndZ, this](const FVector3f& Position, TArray<int32>& PositionMatches, const float ComparisonThreshold)
-		{
-			PositionMatches.Reset();
-			const FIndexAndZ PositionZ = FIndexAndZ(0, Position);
-		
-			// Use binary search to narrow down the range that matches our search area.
-			FIndexAndZ ZBound;
-			ZBound.Z = PositionZ.Z - ComparisonThreshold;
-			const int32 StartRange = Algo::LowerBound(VertIndexAndZ, ZBound, FCompareIndexAndZ{});
-
-			ZBound.Z = PositionZ.Z + ComparisonThreshold;
-			const int32 EndRange = Algo::UpperBound(VertIndexAndZ, ZBound, FCompareIndexAndZ{});
-		
-			// Search for duplicates, quickly!
-			for (int32 Index = StartRange; Index < EndRange; Index++)
-			{
-				const FVector3f& PositionA = ImportDataDst.Points[VertIndexAndZ[Index].Index];
-				if (PointsEqual(PositionA, Position, ComparisonThreshold))
-				{
-					PositionMatches.Add(VertIndexAndZ[Index].Index);
-				}
-			}
-		};
-
-		TArray<int32> SrcWedgeIndexes;
-		TArray<int32> DestWedgeIndexes;
-		
-		// Match all source vertex with destination vertex
-		for (int32 VertexIndexSrc: InSrcVertices)
-		{
-			const FVector3f& PositionSrc = ImportDataSrc.Points[VertexIndexSrc];
-		
-			TArray<int32> SimilarDestinationVertex;
-			FindSimilarPosition(PositionSrc, SimilarDestinationVertex, KINDA_SMALL_NUMBER);
-
-			if (SimilarDestinationVertex.Num() == 0)
-			{
-				//Match with UV projection
-				VertexIndexToMatchWithPositions.Add(VertexIndexSrc);
-			}
-			else
-			{
-				//We have a direct match
-				VertexMatchNameSpace::FVertexMatchResult& VertexMatchDest = VertexIndexSrcToVertexIndexDestMatches.Add(VertexIndexSrc);
-
-				SrcWedgeIndexes.Reset();
-				VertexToSrcWedgeMap.MultiFind(VertexIndexSrc, SrcWedgeIndexes);
-
-				if (SrcWedgeIndexes.Num() > 0 && SimilarDestinationVertex.Num() > 1)
-				{
-					//Check if we have a point that is perfectly matching (position, UV, material and vertex color). Because normals and tangent are on the triangles we do not test those.
-					for (int32 MatchDestinationIndex = 0; MatchDestinationIndex < SimilarDestinationVertex.Num(); ++MatchDestinationIndex)
-					{
-						int32 VertexIndexDest = SimilarDestinationVertex[MatchDestinationIndex];
-						DestWedgeIndexes.Reset();
-						VertexToDstWedgeMap.MultiFind(VertexIndexDest, DestWedgeIndexes);
-					
-						for (int32 IndexDest = 0; IndexDest < DestWedgeIndexes.Num(); ++IndexDest)
-						{
-							int32 DestWedgeIndex = DestWedgeIndexes[IndexDest];
-							const SkeletalMeshImportData::FVertex& WedgeDest = ImportDataDst.Wedges[DestWedgeIndex];
-							for (int32 IndexSrc = 0; IndexSrc < SrcWedgeIndexes.Num(); ++IndexSrc)
-							{
-								int32 SrcWedgeIndex = SrcWedgeIndexes[IndexSrc];
-								const SkeletalMeshImportData::FVertex& WedgeSrc = ImportDataSrc.Wedges[SrcWedgeIndex];
-								//Wedge == operator test: material, vertex color and UVs
-								if (WedgeDest == WedgeSrc)
-								{
-									VertexMatchDest.VertexIndexes.Add(SimilarDestinationVertex[MatchDestinationIndex]);
-									VertexMatchDest.Ratios.Add(1.0f);
-									break;
-								}
-							}
-							if (VertexMatchDest.VertexIndexes.Num() > 0)
-							{
-								break;
-							}
-						}
-					}
-				}
-				
-				//If there is no direct match, simply put everything
-				if (VertexMatchDest.VertexIndexes.Num() == 0)
-				{
-					for (int32 MatchDestinationIndex = 0; MatchDestinationIndex < SimilarDestinationVertex.Num(); ++MatchDestinationIndex)
-					{
-						VertexMatchDest.VertexIndexes.Add(SimilarDestinationVertex[MatchDestinationIndex]);
-						VertexMatchDest.Ratios.Add(1.0f / SimilarDestinationVertex.Num());
-					}
-				}
-			}
-		}
-	}
-	
-private:
-	const FSkeletalMeshImportData& ImportDataSrc;
-	const FSkeletalMeshImportData& ImportDataDst;
-	const TSet<uint32>& DstVertices;
-	TMultiMap<int32, int32> VertexToSrcWedgeMap;
-	TMultiMap<int32, int32> VertexToDstWedgeMap;
-	TArray<uint32>& VertexIndexToMatchWithPositions;
-	TSortedMap<uint32, VertexMatchNameSpace::FVertexMatchResult>& VertexIndexSrcToVertexIndexDestMatches;
-};
-
-uint64 HashFromBoneAndVertex(const int32 InBoneIndex, const int32 InVertexIndex)
-{
-	return ((uint64)(InBoneIndex) << 32 & 0xFFFFFFFF00000000) | ((uint64)(InVertexIndex) & 0x00000000FFFFFFFF);
-}
-
-class FUnusedInfluenceBuilder
-{
-public:
-	FUnusedInfluenceBuilder(
-		const TSortedMap<uint32, VertexMatchNameSpace::FVertexMatchResult>& InVertexIndexSrcToVertexIndexDestMatches,
-		const TMap<uint32, TArray<int32>>& InSourceVertexIndexToAlternateInfluenceIndexMap,
-		const int32 InNumVertices,
-		TArray<SkeletalMeshImportData::FRawBoneInfluence>& InOutAlternateInfluences)
-		: VertexIndexSrcToVertexIndexDestMatches(InVertexIndexSrcToVertexIndexDestMatches)
-		, SourceVertexIndexToAlternateInfluenceIndexMap(InSourceVertexIndexToAlternateInfluenceIndexMap)
-		, NumVertices(InNumVertices)
-		, AlternateInfluences(InOutAlternateInfluences)	
-	{}
-	
-	void TryBuilding(const TSet<uint32>& InUnusedVertices) const
-	{
-		const int32 AlternateInfluenceNum = AlternateInfluences.Num();
-		if (AlternateInfluenceNum == 0 || NumVertices == 0 || InUnusedVertices.IsEmpty())
-		{
-			return;
-		}
-		
-		TArray<bool> SrcVertexIndexMatched;
-		SrcVertexIndexMatched.AddZeroed(NumVertices);
-
-		TMap<uint64, int32> InfluenceKeyToInfluenceIndex;
-		InfluenceKeyToInfluenceIndex.Reserve(AlternateInfluenceNum);
-	
-		for (int32 AlternateInfluenceIndex = 0; AlternateInfluenceIndex < AlternateInfluenceNum; ++AlternateInfluenceIndex)
-		{
-			SkeletalMeshImportData::FRawBoneInfluence& Influence = AlternateInfluences[AlternateInfluenceIndex];
-			SrcVertexIndexMatched[Influence.VertexIndex] = true;
-			const uint64 Key = HashFromBoneAndVertex(Influence.BoneIndex, Influence.VertexIndex);
-			InfluenceKeyToInfluenceIndex.Add(Key, AlternateInfluenceIndex);
-		}
-
-		for (const uint32 Unused: InUnusedVertices)
-		{
-			//Skip if the vertex is already matched
-			if (SrcVertexIndexMatched[Unused])
-			{
-				continue;
-			}
-
-			const VertexMatchNameSpace::FVertexMatchResult* SrcVertexMatch = VertexIndexSrcToVertexIndexDestMatches.Find(Unused);
-			if (!SrcVertexMatch || SrcVertexMatch->VertexIndexes.IsEmpty())
-			{
-				//No match skip this influence
-				continue;
-			}
-		
-			for (int32 ImpactedIndex = 0; ImpactedIndex < SrcVertexMatch->VertexIndexes.Num(); ++ImpactedIndex)
-			{
-				float Ratio = SrcVertexMatch->Ratios[ImpactedIndex];
-				if (!FMath::IsNearlyZero(Ratio, KINDA_SMALL_NUMBER))
-				{
-					uint32 VertexIndexSrc = SrcVertexMatch->VertexIndexes[ImpactedIndex];
-					const TArray<int32>* SrcAlternateInfluences = SourceVertexIndexToAlternateInfluenceIndexMap.Find(VertexIndexSrc);
-					if (!SrcAlternateInfluences)
-					{
-						continue;
-					}
-
-					for (const int32 AlternateInfluenceIndex: *SrcAlternateInfluences)
-					{
-						if (!AlternateInfluences.IsValidIndex(AlternateInfluenceIndex))
-						{
-							continue;
-						}
-					
-						SrcVertexIndexMatched[Unused] = true;
-						SkeletalMeshImportData::FRawBoneInfluence AlternateInfluence = AlternateInfluences[AlternateInfluenceIndex];
-						const uint64 Key = HashFromBoneAndVertex(AlternateInfluence.BoneIndex, Unused);
-					
-						if (!InfluenceKeyToInfluenceIndex.Contains(Key))
-						{
-							AlternateInfluence.VertexIndex = Unused;
-							InfluenceKeyToInfluenceIndex.Add(Key, AlternateInfluences.Add(AlternateInfluence));
-						}
-						else
-						{
-							int32& InfluenceIndex = InfluenceKeyToInfluenceIndex.FindOrAdd(Key);
-							SkeletalMeshImportData::FRawBoneInfluence& ExistAlternateInfluence = AlternateInfluences[InfluenceIndex];
-							if (ExistAlternateInfluence.Weight < AlternateInfluence.Weight)
-							{
-								ExistAlternateInfluence.Weight = AlternateInfluence.Weight;
-							}
-						}
-					}
-				}
-			}
-		}
-	}
-	
-private:
-	
-	const TSortedMap<uint32, VertexMatchNameSpace::FVertexMatchResult>& VertexIndexSrcToVertexIndexDestMatches;
-	const TMap<uint32, TArray<int32>>& SourceVertexIndexToAlternateInfluenceIndexMap;
-	int32 NumVertices = 0;
-	TArray<SkeletalMeshImportData::FRawBoneInfluence>& AlternateInfluences;
-};
-	
-}
-
 bool FLODUtilities::UpdateAlternateSkinWeights(
 	FSkeletalMeshLODModel& LODModelDest,
 	FSkeletalMeshImportData& ImportDataDest,
@@ -2758,267 +2309,31 @@ bool FLODUtilities::UpdateAlternateSkinWeights(
 	int32 LODIndexDest,
 	const IMeshUtilities::MeshBuildOptions& Options)
 {
-	using namespace AlternateSkinWeightsLocal;
-	
-	//Ensure log message only once
-	bool bNoMatchMsgDone = false;
-	const int32 PointNumberDest = ImportDataDest.Points.Num();
-
 	int32 ProfileIndex = 0;
 	if (!ImportDataDest.AlternateInfluenceProfileNames.Find(ProfileNameDest.ToString(), ProfileIndex))
 	{
-		FFormatNamedArguments Args;
-		Args.Add(TEXT("SkeletalMeshName"), FText::FromString(SkeletalMeshDest->GetName()));
-		Args.Add(TEXT("ProfileName"), FText::FromName(ProfileNameDest));
-		Args.Add(TEXT("LODIndex"), LODIndexDest);
-
-		FText Message = FText::Format(NSLOCTEXT("FLODUtilities_UpdateAlternateSkinWeights", "AlternateDataNotAvailable", 
-			"Asset {SkeletalMeshName} LOD {LODIndex} failed to import skin weight profile {ProfileName}. The alternate skinning imported source data is not available."), 
-			Args);
-		UE_LOG(LogLODUtilities, Display, TEXT("%s"), *(Message.ToString()));
+		UE_LOGFMT(LogLODUtilities, Display, "Can't update alternative profile {ProfileName} on {SkeletalMesh} at LOD {0}. Profile doesn't exist.",
+			ProfileNameDest, SkeletalMeshDest->GetName(), LODIndexDest);
 		return false;
 	}
 
 	check(ImportDataDest.AlternateInfluences.IsValidIndex(ProfileIndex));
+	
 	//The data must be there and must be verified before getting here
 	const FSkeletalMeshImportData& ImportDataSrc = ImportDataDest.AlternateInfluences[ProfileIndex];
-	const int32 PointNumberSrc = ImportDataSrc.Points.Num();
-	const int32 InfluenceNumberSrc = ImportDataSrc.Influences.Num();
-
-	if (PointNumberDest != PointNumberSrc)
-	{
-		FFormatNamedArguments Args;
-		Args.Add(TEXT("SkeletalMeshName"), FText::FromString(SkeletalMeshDest->GetName()));
-		Args.Add(TEXT("PointNumberSrc"), PointNumberSrc);
-		Args.Add(TEXT("PointNumberDest"), PointNumberDest);
-
-		FText Message = FText::Format(NSLOCTEXT("FLODUtilities_UpdateAlternateSkinWeights", "DifferentPointNumber", "Asset {SkeletalMeshName} failed to import skin weight profile the alternate skinning model has a different number of vertex. Alternate vertex count: {PointNumberSrc}, LOD vertex count: {PointNumberDest}"), Args);
-		UE_LOG(LogLODUtilities, Warning, TEXT("%s"), *(Message.ToString()));
-		return false;
-	}
 
-	// to avoid building render data using 'unused' vertices (not referenced by any triangle),
-	// retrieve the vertices info from the triangles
-	TSet<uint32> SrcFaceVert;
-	GetVerticesFromFaces(ImportDataSrc, SrcFaceVert);
-	
-	TSet<uint32> DstFaceVert;
-	GetVerticesFromFaces(ImportDataDest, DstFaceVert);
-
-	if (DstFaceVert.Num() != SrcFaceVert.Num())
+	// Double-check the mesh has the same structure.
+	if (!IsSameImportMeshStructure(ImportDataSrc, ImportDataDest))
 	{
-		NotifyFromFacesDiscrepancy(SkeletalMeshDest, SrcFaceVert.Num(), DstFaceVert.Num());
+		UE_LOGFMT(LogLODUtilities, Warning, "Can't update alternative profile {ProfileName} on {SkeletalMesh} at LOD {0}. Profile source mesh and the base mesh are different.",
+			ProfileNameDest, SkeletalMeshDest->GetName(), LODIndexDest);
 		return false;
 	}
 
-	// build unused vertices data
-	const int32 NumUsedVertices = DstFaceVert.Num();
-	const int32 NumUnusedVertices = PointNumberDest - NumUsedVertices; 
-	TSet<uint32> UnusedSrcFaceVert, UnusedDstFaceVert;
-	if (NumUnusedVertices > 0)
-	{
-		UE_LOG(LogLODUtilities, Display, TEXT("Alternate skinning import: %d vertices do not reference any triangle."), NumUnusedVertices);
-		GetUnusedVertices(SrcFaceVert, ImportDataSrc.Points.Num(), UnusedSrcFaceVert);
-		GetUnusedVertices(DstFaceVert, ImportDataDest.Points.Num(), UnusedDstFaceVert);
-	}
-	
-	//Create a map to remap source bone index to destination bone index
-	TMap<int32, int32> RemapBoneIndexSrcToDest;
-	FillRemapBoneIndexSrcToDest(ImportDataSrc, ImportDataDest, RemapBoneIndexSrcToDest);
-
-	//Map to get the vertex index source to a destination vertex match
-	TSortedMap<uint32, VertexMatchNameSpace::FVertexMatchResult> VertexIndexSrcToVertexIndexDestMatches;
-	VertexIndexSrcToVertexIndexDestMatches.Reserve(DstFaceVert.Num());
-	TArray<uint32> VertexIndexToMatchWithPositions;
-
-	// Match all used source vertex with destination vertex
-	const FVertexMatcher Matcher(ImportDataSrc, ImportDataDest, DstFaceVert, VertexIndexToMatchWithPositions, VertexIndexSrcToVertexIndexDestMatches);
-	Matcher.TryMatching(SrcFaceVert);
-	
-	//Find a match for all unmatched source vertex, unmatched vertex happen when the geometry is different between source and destination mesh
-	bool bAllSourceVertexAreMatch = VertexIndexToMatchWithPositions.Num() <= 0 && VertexIndexSrcToVertexIndexDestMatches.Num() == NumUsedVertices;
-	if (!bAllSourceVertexAreMatch)
-	{
-		MatchVertexIndexUsingPosition(ImportDataDest, ImportDataSrc, VertexIndexSrcToVertexIndexDestMatches, VertexIndexToMatchWithPositions, bNoMatchMsgDone);
-		//Make sure each used vertex index source has a match, warn the user in case there is no match
-		for (int32 VertexIndexSource: SrcFaceVert)
-		{
-			if (!VertexIndexSrcToVertexIndexDestMatches.Contains(VertexIndexSource))
-			{
-				//Skip this vertex, its possible the skinning quality can be affected here
-				if (!bNoMatchMsgDone)
-				{
-					UE_LOG(LogLODUtilities, Warning, TEXT("Alternate skinning import: Cannot find a destination vertex index match for source vertex index. Alternate skinning quality will be lower."));
-					bNoMatchMsgDone = true;
-				}
-			}
-		}
-		bAllSourceVertexAreMatch = VertexIndexSrcToVertexIndexDestMatches.Num() == NumUsedVertices;
-	}
-	
-	
-	//Find the Destination to source match, to make sure all extra destination vertex get weighted properly in the alternate influences
-	TSortedMap<uint32, VertexMatchNameSpace::FVertexMatchResult> VertexIndexDestToVertexIndexSrcMatches;
-	if(!bAllSourceVertexAreMatch)
-	{
-		VertexIndexDestToVertexIndexSrcMatches.Reserve(NumUsedVertices);
-		const TArray<uint32> VertexIndexToMatch = DstFaceVert.Array();
-		MatchVertexIndexUsingPosition(ImportDataSrc, ImportDataDest, VertexIndexDestToVertexIndexSrcMatches, VertexIndexToMatch, bNoMatchMsgDone);
-	}
-
-	//We now iterate the source influence and create the alternate influence by using the matches between source and destination vertex
-	TArray<SkeletalMeshImportData::FRawBoneInfluence> AlternateInfluences;
-	AlternateInfluences.Empty(ImportDataSrc.Influences.Num());
-
-	TMap<uint32, TArray<int32>> SourceVertexIndexToAlternateInfluenceIndexMap;
-	SourceVertexIndexToAlternateInfluenceIndexMap.Reserve(InfluenceNumberSrc);
-	
-	for (int32 InfluenceIndexSrc = 0; InfluenceIndexSrc < InfluenceNumberSrc; ++InfluenceIndexSrc)
-	{
-		const SkeletalMeshImportData::FRawBoneInfluence& InfluenceSrc = ImportDataSrc.Influences[InfluenceIndexSrc];
-		int32 VertexIndexSource = InfluenceSrc.VertexIndex;
-		int32 BoneIndexSource = InfluenceSrc.BoneIndex;
-		//We need to remap the source bone index to have the matching target bone index
-		if (const int32* BoneIndexDstPtr = RemapBoneIndexSrcToDest.Find(BoneIndexSource))
-		{
-			const int32 BoneIndexDest = *BoneIndexDstPtr;
-			//Find the match destination vertex index
-			VertexMatchNameSpace::FVertexMatchResult* SourceVertexMatch = VertexIndexSrcToVertexIndexDestMatches.Find(VertexIndexSource);
-			if (SourceVertexMatch == nullptr || SourceVertexMatch->VertexIndexes.Num() <= 0)
-			{
-				//No match skip this influence
-				continue;
-			}
-			TArray<int32>& AlternateInfluencesMap = SourceVertexIndexToAlternateInfluenceIndexMap.FindOrAdd(VertexIndexSource);
-			//No need to merge all vertexindex per bone, ProcessImportMeshInfluences will do this for us later
-			//So just add all of the entry we have.
-			for (int32 ImpactedIndex = 0; ImpactedIndex < SourceVertexMatch->VertexIndexes.Num(); ++ImpactedIndex)
-			{
-				uint32 VertexIndexDest = SourceVertexMatch->VertexIndexes[ImpactedIndex];
-				const float Ratio = SourceVertexMatch->Ratios[ImpactedIndex];
-
-				const float InfluenceWeight = InfluenceSrc.Weight* Ratio;
-				if (InfluenceWeight < UE::AnimationCore::BoneWeightThreshold)
-				{
-					continue;
-				}
-				
-				SkeletalMeshImportData::FRawBoneInfluence AlternateInfluence;
-				AlternateInfluence.BoneIndex = BoneIndexDest;
-				AlternateInfluence.VertexIndex = VertexIndexDest;
-				AlternateInfluence.Weight = InfluenceWeight;
-				int32 AlternateInfluencesIndex = AlternateInfluences.Add(AlternateInfluence);
-				AlternateInfluencesMap.Add(AlternateInfluencesIndex);
-			}
-		}
-	}
-
-	//In case the source geometry was not matching the destination we have to add influence for each extra destination vertex index
-	if (VertexIndexDestToVertexIndexSrcMatches.Num() > 0)
-	{
-		TBitArray<> DestinationVertexIndexMatched(false, PointNumberDest);
-		const int32 AlternateInfluenceNumber = AlternateInfluences.Num();
-		
-		//We want to avoid making duplicate so we use a map where the key is the boneindex mix with the destination vertex index
-		TMap<uint64, int32> InfluenceKeyToInfluenceIndex;
-		InfluenceKeyToInfluenceIndex.Reserve(AlternateInfluenceNumber);
-		for (int32 AlternateInfluenceIndex = 0; AlternateInfluenceIndex < AlternateInfluenceNumber; ++AlternateInfluenceIndex)
-		{
-			SkeletalMeshImportData::FRawBoneInfluence& Influence = AlternateInfluences[AlternateInfluenceIndex];
-			DestinationVertexIndexMatched[Influence.VertexIndex] = true;
-			uint64 Key = ((uint64)(Influence.BoneIndex) << 32 & 0xFFFFFFFF00000000) | ((uint64)(Influence.VertexIndex) & 0x00000000FFFFFFFF);
-			InfluenceKeyToInfluenceIndex.Add(Key, AlternateInfluenceIndex);
-		}
-
-		for (int32 VertexIndexDestination = 0; VertexIndexDestination < PointNumberDest; ++VertexIndexDestination)
-		{
-			//Skip if the vertex is already matched
-			if (DestinationVertexIndexMatched[VertexIndexDestination])
-			{
-				continue;
-			}
-			VertexMatchNameSpace::FVertexMatchResult* DestinationVertexMatch = VertexIndexDestToVertexIndexSrcMatches.Find(VertexIndexDestination);
-			if (DestinationVertexMatch == nullptr || DestinationVertexMatch->VertexIndexes.Num() <= 0)
-			{
-				//No match skip this influence
-				continue;
-			}
-			for (int32 ImpactedIndex = 0; ImpactedIndex < DestinationVertexMatch->VertexIndexes.Num(); ++ImpactedIndex)
-			{
-				uint32 VertexIndexSrc = DestinationVertexMatch->VertexIndexes[ImpactedIndex];
-				float Ratio = DestinationVertexMatch->Ratios[ImpactedIndex];
-				if (!FMath::IsNearlyZero(Ratio, KINDA_SMALL_NUMBER))
-				{
-					//Find src influence for this source vertex index
-					TArray<int32>* AlternateInfluencesMap = SourceVertexIndexToAlternateInfluenceIndexMap.Find(VertexIndexSrc);
-					if (AlternateInfluencesMap == nullptr)
-					{
-						continue;
-					}
-					for (int32 AlternateInfluencesMapIndex = 0; AlternateInfluencesMapIndex < (*AlternateInfluencesMap).Num(); ++AlternateInfluencesMapIndex)
-					{
-						int32 AlternateInfluenceIndex = (*AlternateInfluencesMap)[AlternateInfluencesMapIndex];
-						if (!AlternateInfluences.IsValidIndex(AlternateInfluenceIndex))
-						{
-							continue;
-						}
-						DestinationVertexIndexMatched[VertexIndexDestination] = true;
-						SkeletalMeshImportData::FRawBoneInfluence AlternateInfluence = AlternateInfluences[AlternateInfluenceIndex];
-						uint64 Key = ((uint64)(AlternateInfluence.BoneIndex) << 32 & 0xFFFFFFFF00000000) | ((uint64)(VertexIndexDestination) & 0x00000000FFFFFFFF);
-						if (!InfluenceKeyToInfluenceIndex.Contains(Key))
-						{
-							AlternateInfluence.VertexIndex = VertexIndexDestination;
-							InfluenceKeyToInfluenceIndex.Add(Key, AlternateInfluences.Add(AlternateInfluence));
-						}
-						else
-						{
-							int32& InfluenceIndex = InfluenceKeyToInfluenceIndex.FindOrAdd(Key);
-							SkeletalMeshImportData::FRawBoneInfluence& ExistAlternateInfluence = AlternateInfluences[InfluenceIndex];
-							if (ExistAlternateInfluence.Weight < AlternateInfluence.Weight)
-							{
-								ExistAlternateInfluence.Weight = AlternateInfluence.Weight;
-							}
-						}
-					}
-				}
-			}
-		}
-	}
-	
-	// Check if there are any points that have no weight assignments after all is said and done. Attempt to average across neighbors to reconstruct
-	// those weights. Currently this relies on these points being completely surrounded by at least two points that have weights assigned to them.
-	TBitArray<> UnassignedVertexBitmask(false, ImportDataDest.Points.Num());
-	for (const SkeletalMeshImportData::FVertex& Wedge: ImportDataDest.Wedges)
-	{
-		UnassignedVertexBitmask[Wedge.VertexIndex] = true;
-	}
-
-	// mark unused vertices as assigned to basically skip them
-	for (uint32 UnusedDstVertexIndex: UnusedDstFaceVert)
-	{
-		UnassignedVertexBitmask[UnusedDstVertexIndex] = false;
-	}
-
-	for (const SkeletalMeshImportData::FRawBoneInfluence& Influence: AlternateInfluences)
-	{
-		UnassignedVertexBitmask[Influence.VertexIndex] = false;
-	}
-
-	TArray<int32> UnassignedVertexes;
-	for (TConstSetBitIterator<> It(UnassignedVertexBitmask); It; ++It)
-	{
-		UnassignedVertexes.Add(It.GetIndex());
-	}
-
-	if (!UnassignedVertexes.IsEmpty())
-	{
-		InfillMissingVertexSkinWeights(ImportDataDest, UnassignedVertexes, AlternateInfluences);
-	}
-
 	//Sort and normalize weights for alternate influences
-	ProcessImportMeshInfluences(ImportDataDest.Wedges.Num(), AlternateInfluences, SkeletalMeshDest->GetPathName());
-
+	TArray<SkeletalMeshImportData::FRawBoneInfluence> AlternateInfluences{ImportDataSrc.Influences};
 	
+	ProcessImportMeshInfluences(ImportDataDest.Wedges.Num(), AlternateInfluences, SkeletalMeshDest->GetPathName());
 
 	//Store the remapped influence into the profile, the function SkeletalMeshTools::ChunkSkinnedVertices will use all profiles including this one to chunk the sections
 	FImportedSkinWeightProfileData& ImportedProfileData = LODModelDest.SkinWeightProfiles.Add(ProfileNameDest);
@@ -3036,7 +2351,6 @@ bool FLODUtilities::UpdateAlternateSkinWeights(
 	//
 	//////////////////////////////////////////////////////////////////////////
 
-	bool bBuildSuccess = true;
 	//Prepare the build data to rebuild the asset with the alternate influences
 	//The chunking can be different when we have alternate influences
 	//Grab the build data from ImportDataDest
@@ -3062,7 +2376,7 @@ bool FLODUtilities::UpdateAlternateSkinWeights(
 	AdjustImportDataFaceMaterialIndex(SkeletalMeshDest->GetMaterials(), ImportDataDest.Materials, LODFacesDest, LODIndexDest);
 
 	//Build the destination mesh with the Alternate influences, so the chunking is done properly.
-	bBuildSuccess = MeshUtilities.BuildSkeletalMesh(LODModelDest, SkeletalMeshDest->GetName(), RefSkeleton, LODInfluencesDest, LODWedgesDest, LODFacesDest, LODPointsDest, LODPointToRawMapDest, BuildOptions, &WarningMessages, &WarningNames);
+	bool bBuildSuccess = MeshUtilities.BuildSkeletalMesh(LODModelDest, SkeletalMeshDest->GetName(), RefSkeleton, LODInfluencesDest, LODWedgesDest, LODFacesDest, LODPointsDest, LODPointToRawMapDest, BuildOptions, &WarningMessages, &WarningNames);
 	//Re-Apply the user section changes, the UserSectionsData is map to original section and should match the builded LODModel
 	LODModelDest.SyncronizeUserSectionsDataArray();
 
@@ -3243,8 +2557,7 @@ void FLODUtilities::GenerateImportedSkinWeightProfileData(
 		
 		FSkelMeshSection& Section = LODModelDest.Sections[SectionIndex];
 		const TArray<FBoneIndexType>& SectionBoneMap = Section.BoneMap;
-		const int32 VertexIndex = LODModelDest.MeshToImportVertexMap[VertexInstanceIndex];
-		check(VertexIndex >= 0 && VertexIndex <= LODModelDest.MaxImportVertex);
+		const int32 VertexIndex = LODModelDest.GetRawPointIndices()[VertexInstanceIndex];
 		FRawSkinWeight& SkinWeight = SkinWeights.AddDefaulted_GetRef();
 
 		// get the bone weights for that vertex and sort them (mostly by weights here as VertIndex are similar)
diff --git a/Engine/Source/Runtime/Core/Private/UObject/FortniteReleaseBranchCustomObjectVersion.cpp b/Engine/Source/Runtime/Core/Private/UObject/FortniteReleaseBranchCustomObjectVersion.cpp
--- a/Engine/Source/Runtime/Core/Private/UObject/FortniteReleaseBranchCustomObjectVersion.cpp
+++ b/Engine/Source/Runtime/Core/Private/UObject/FortniteReleaseBranchCustomObjectVersion.cpp
@@ -7,7 +7,7 @@ TMap<FGuid, FGuid> FFortniteReleaseBranchCustomObjectVersion::GetSystemGuids()
 	const FDevSystemGuids& DevGuids = FDevSystemGuids::Get();
 
 	SystemGuids.Add(DevGuids.Niagara_LatestScriptCompileVersion, FGuid("D35A655CD3FDA4408724D77E1316E35B"));
-	SystemGuids.Add(DevGuids.SkeletalMeshDerivedDataVersion, FGuid("5DDF529BE4164CF6860F00156570398A"));
+	SystemGuids.Add(DevGuids.SkeletalMeshDerivedDataVersion, FGuid("C96DEBAD45E3443D804017D4C29F7C4B"));
 	SystemGuids.Add(DevGuids.MaterialTranslationDDCVersion, FGuid("888A502A66F6401580B27451F211C46B"));
 	SystemGuids.Add(DevGuids.GLOBALSHADERMAP_DERIVEDDATA_VER, FGuid("0A0A1BF5D2A94105BB395EF0DA7CAEB9"));
 	SystemGuids.Add(DevGuids.MATERIALSHADERMAP_DERIVEDDATA_VER, FGuid("C0E1F59386894C9D855B85178B2553CB"));
