Interchange: Fix skeletal mesh re-import material assignment issues.
if you re-import a skeletalmesh that has a disabled section with a new file containing more material sections and two rename material section, the re-import was not handling correctly the material slot and the lod section count.

#rnx

diff --git a/Engine/Source/Developer/SkeletalMeshUtilitiesCommon/Private/LODUtilities.cpp b/Engine/Source/Developer/SkeletalMeshUtilitiesCommon/Private/LODUtilities.cpp
--- a/Engine/Source/Developer/SkeletalMeshUtilitiesCommon/Private/LODUtilities.cpp
+++ b/Engine/Source/Developer/SkeletalMeshUtilitiesCommon/Private/LODUtilities.cpp
@@ -4037,17 +4037,39 @@ void FLODUtilities::MatchImportedMaterials(FLODUtilities::FSkeletalMeshMatchImpo
 		else
 		{
 			ExistLODModel.UserSectionsData.Reset();
+			TBitArray<> MatchedMaterial;
+			MatchedMaterial.Init(false, Materials.Num());
+			TSet<FName> MatchedUserSections;
+			MatchedUserSections.Reserve(NewUserSectionsDataMap.Num());
 			for (TPair<FName, FSkelMeshSourceSectionUserData> UserSectionsData : NewUserSectionsDataMap)
 			{
 				for (int32 MaterialIndex = 0; MaterialIndex < Materials.Num(); ++MaterialIndex)
 				{
 					if (Materials[MaterialIndex].ImportedMaterialSlotName == UserSectionsData.Key)
 					{
+						MatchedMaterial[MaterialIndex] = true;
+						MatchedUserSections.Add(UserSectionsData.Key);
 						ExistLODModel.UserSectionsData.Add(MaterialIndex, UserSectionsData.Value);
 						break;
 					}
 				}
 			}
+			//Add any missing user section data
+			for (TPair<FName, FSkelMeshSourceSectionUserData> UserSectionsData : NewUserSectionsDataMap)
+			{
+				if (!MatchedUserSections.Contains(UserSectionsData.Key))
+				{
+					for (int32 MaterialIndex = 0; MaterialIndex < Materials.Num(); ++MaterialIndex)
+					{
+						if (!MatchedMaterial[MaterialIndex])
+						{
+							MatchedMaterial[MaterialIndex] = true;
+							ExistLODModel.UserSectionsData.Add(MaterialIndex, UserSectionsData.Value);
+							break;
+						}
+					}
+				}
+			}
 			//Reorder the user section from the keys
 			ExistLODModel.UserSectionsData.KeySort([](const int32& A, const int32& B)
 				{
@@ -4080,6 +4102,13 @@ void FLODUtilities::ReorderMaterialSlotToBaseLod(USkeletalMesh* SkeletalMesh)
 	TArray<FSkeletalMaterial> ReorderMaterialArray;
 	ReorderMaterialArray.Reserve(Materials.Num());
 
+	//Unique imported material list for all lods
+	TArray<FName> AllLodSortedImportedMaterials;
+	AllLodSortedImportedMaterials.Reserve(Materials.Num());
+
+	TBitArray<> MatchedMaterials;
+	MatchedMaterials.Init(false, Materials.Num());
+
 	for (int32 LodIndex = 0; LodIndex < SkeletalMesh->GetLODNum(); ++LodIndex)
 	{
 		if (!SkeletalMesh->HasMeshDescription(LodIndex))
@@ -4098,33 +4127,71 @@ void FLODUtilities::ReorderMaterialSlotToBaseLod(USkeletalMesh* SkeletalMesh)
 		PRAGMA_ENABLE_DEPRECATION_WARNINGS
 		const TArray<SkeletalMeshImportData::FMaterial>& ImportedMaterials = LodImportData.Materials;
 
-		//Find remap index for this LOD matching materials
 		for (int32 ImportedMaterialIndex = 0; ImportedMaterialIndex < ImportedMaterials.Num(); ++ImportedMaterialIndex)
 		{
 			FName ImportedMaterialSlotName = FName(*ImportedMaterials[ImportedMaterialIndex].MaterialImportName);
-			for (int32 MaterialIndex = 0; MaterialIndex < Materials.Num(); ++MaterialIndex)
+			//add Unique Material
+			if (!AllLodSortedImportedMaterials.Contains(ImportedMaterialSlotName))
 			{
-				FName MaterialSlotName = Materials[MaterialIndex].ImportedMaterialSlotName;
-				if (MaterialSlotRemap[MaterialIndex] != INDEX_NONE)
+				AllLodSortedImportedMaterials.Add(ImportedMaterialSlotName);
+				//Make sure matched material is filled
+				for (int32 MaterialIndex = 0; MaterialIndex < Materials.Num(); ++MaterialIndex)
 				{
-					//If the name match say we found the match)
+					FName MaterialSlotName = Materials[MaterialIndex].ImportedMaterialSlotName;
 					if (MaterialSlotName == ImportedMaterialSlotName)
 					{
+						MatchedMaterials[MaterialIndex] = true;
 						break;
 					}
-					continue;
 				}
-				int32& RemapIndex = MaterialSlotRemap[MaterialIndex];
+			}
+		}
+	}
 
+	//Now fill the ReorderMaterialArray using the AllLodSortedImportedMaterials order
+	for (int32 ImportedMaterialIndex = 0; ImportedMaterialIndex < AllLodSortedImportedMaterials.Num(); ++ImportedMaterialIndex)
+	{
+		FName ImportedMaterialSlotName = AllLodSortedImportedMaterials[ImportedMaterialIndex];
+		int32 TargetMaterialIndex = INDEX_NONE;
+		for (int32 MaterialIndex = 0; MaterialIndex < Materials.Num(); ++MaterialIndex)
+		{
+			//If we have a match we should find back the material index and assign it to the target material index
+			if (MatchedMaterials[MaterialIndex])
+			{
+				FName MaterialSlotName = Materials[MaterialIndex].ImportedMaterialSlotName;
 				if (MaterialSlotName == ImportedMaterialSlotName)
 				{
-					RemapIndex = ReorderMaterialArray.Add(Materials[MaterialIndex]);
+					TargetMaterialIndex = MaterialIndex;
 					break;
 				}
 			}
 		}
+		//If we have an unmatched imported material we use the next unmatched material in the original list
+		if (TargetMaterialIndex == INDEX_NONE)
+		{
+			for (int32 MaterialIndex = 0; MaterialIndex < Materials.Num(); ++MaterialIndex)
+			{
+				if (MatchedMaterials[MaterialIndex])
+				{
+					continue;
+				}
+				ensure(MaterialSlotRemap[MaterialIndex] == INDEX_NONE);
+				TargetMaterialIndex = MaterialIndex;
+				MatchedMaterials[MaterialIndex] = true;
+				break;
+			}
+		}
+		//We should have found a valid material index at this point
+		if (ensure(MaterialSlotRemap.IsValidIndex(TargetMaterialIndex)))
+		{
+			if (MaterialSlotRemap[TargetMaterialIndex] == INDEX_NONE)
+			{
+				MaterialSlotRemap[TargetMaterialIndex] = ReorderMaterialArray.Add(Materials[TargetMaterialIndex]);
+			}
+		}
 	}
-	//Custom LOD can add materials, so we add them at the end of the material slots
+
+	//Unmatched material are appended to the end of the material slots
 	for (int32 MaterialIndex = 0; MaterialIndex < Materials.Num(); ++MaterialIndex)
 	{
 		if (MaterialSlotRemap[MaterialIndex] == INDEX_NONE)
