UE: Add name of skeletal mesh to warning about Barycentric calculation

diff --git a/Engine/Source/Developer/SkeletalMeshUtilitiesCommon/Private/LODUtilities.cpp b/Engine/Source/Developer/SkeletalMeshUtilitiesCommon/Private/LODUtilities.cpp
--- a/Engine/Source/Developer/SkeletalMeshUtilitiesCommon/Private/LODUtilities.cpp
+++ b/Engine/Source/Developer/SkeletalMeshUtilitiesCommon/Private/LODUtilities.cpp
@@ -1048,7 +1048,7 @@ bool PointInTriangle(const FVector3f& A, const FVector3f& B, const FVector3f& C,
 	return false;
 }
 
-FVector3f GetBaryCentric(const FVector3f& Point, const FVector3f& A, const FVector3f& B, const FVector3f& C)
+FVector3f GetBaryCentric(const FVector3f& Point, const FVector3f& A, const FVector3f& B, const FVector3f& C, const TCHAR* DebugContext)
 {
 	// Compute the normal of the triangle
 	const FVector3f TriNorm = (B - A) ^ (C - A);
@@ -1069,7 +1069,12 @@ FVector3f GetBaryCentric(const FVector3f& Point, const FVector3f& A, const FVect
 		}
 		return FVector3f(0.0f, 0.0f, 1.0f);
 	}
-	return (FVector3f)FMath::ComputeBaryCentric2D((FVector)Point, (FVector)A, (FVector)B, (FVector)C);
+	FVector ToRet(0, 0, 0);
+	if (!FMath::ComputeBarycentricTri((FVector)Point, (FVector)A, (FVector)B, (FVector)C, ToRet, UE_DOUBLE_SMALL_NUMBER))
+	{
+		UE_LOG(LogUnrealMath, Warning, TEXT("SkeletalMesh [%s]: Small triangle detected in FMath::ComputeBaryCentric2D(); can't compute valid barycentric coordinate."), DebugContext ? DebugContext : TEXT("Unknown Source"));
+	}
+	return (FVector3f)ToRet;
 }
 
 struct FTriangleElement
@@ -1289,7 +1294,7 @@ void ProjectTargetOnBase(const TArray<FSoftSkinVertex>& BaseVertices, const TArr
 				FTriangleElement& BestTriangle = Triangles[FoundIndexMatch];
 				//Found the surface area of the 3 barycentric triangles from the UVs
 				FVector3f BarycentricWeight;
-				BarycentricWeight = GetBaryCentric(FVector3f(FVector2f(TargetUV), 0.0f), FVector3f(BestTriangle.Vertices[0].UVs[0], 0.0f), FVector3f(BestTriangle.Vertices[1].UVs[0], 0.0f), FVector3f(BestTriangle.Vertices[2].UVs[0], 0.0f));	// LWC_TODO: Precision loss
+				BarycentricWeight = GetBaryCentric(FVector3f(FVector2f(TargetUV), 0.0f), FVector3f(BestTriangle.Vertices[0].UVs[0], 0.0f), FVector3f(BestTriangle.Vertices[1].UVs[0], 0.0f), FVector3f(BestTriangle.Vertices[2].UVs[0], 0.0f), DebugContext);	// LWC_TODO: Precision loss
 				//Fill the target match
 				for (int32 Corner = 0; Corner < 3; ++Corner)
 				{
@@ -2042,7 +2047,8 @@ void MatchVertexIndexUsingPosition(
 	, const FSkeletalMeshImportData& ImportDataSrc
 	, TSortedMap<uint32, VertexMatchNameSpace::FVertexMatchResult>& VertexIndexSrcToVertexIndexDestMatches
 	, const TArray<uint32>& VertexIndexToMatchWithPositions
-	, bool& bNoMatchMsgDone)
+	, bool& bNoMatchMsgDone
+	, const TCHAR* DebugContext)
 {
 	if (VertexIndexToMatchWithPositions.Num() <= 0)
 	{
@@ -2097,11 +2103,11 @@ void MatchVertexIndexUsingPosition(
 
 	//This lambda store a source vertex index -> source wedge index destination triangle.
 	//It use a barycentric function to determine the impact on the 3 corner of the triangle.
-	auto AddMatchTriangle = [&ImportDataDest, &TrianglesDest, &VertexIndexSrcToVertexIndexDestMatches](const FTriangleElement& BestTriangle, const FVector3f& Position, const uint32 VertexIndexSrc)
+	auto AddMatchTriangle = [&ImportDataDest, &TrianglesDest, &VertexIndexSrcToVertexIndexDestMatches, DebugContext](const FTriangleElement& BestTriangle, const FVector3f& Position, const uint32 VertexIndexSrc)
 	{
 		//Found the surface area of the 3 barycentric triangles from the UVs
 		FVector3f BarycentricWeight;
-		BarycentricWeight = GetBaryCentric(Position, BestTriangle.Vertices[0].Position, BestTriangle.Vertices[1].Position, BestTriangle.Vertices[2].Position);
+		BarycentricWeight = GetBaryCentric(Position, BestTriangle.Vertices[0].Position, BestTriangle.Vertices[1].Position, BestTriangle.Vertices[2].Position, DebugContext);
 		//Fill the match
 		VertexMatchNameSpace::FVertexMatchResult& VertexMatchDest = VertexIndexSrcToVertexIndexDestMatches.FindOrAdd(VertexIndexSrc);
 		for (int32 CornerIndex = 0; CornerIndex < 3; ++CornerIndex)
