Add cache optimize flag to skelmesh construction

diff --git a/Engine/Source/Developer/SkeletalMeshUtilitiesCommon/Private/StaticToSkeletalMeshConverter.cpp b/Engine/Source/Developer/SkeletalMeshUtilitiesCommon/Private/StaticToSkeletalMeshConverter.cpp
--- a/Engine/Source/Developer/SkeletalMeshUtilitiesCommon/Private/StaticToSkeletalMeshConverter.cpp
+++ b/Engine/Source/Developer/SkeletalMeshUtilitiesCommon/Private/StaticToSkeletalMeshConverter.cpp
@@ -188,8 +188,8 @@ static void CopyReductionSettings(
 static bool AddLODFromMeshDescription(
 	FMeshDescription&& InMeshDescription,
 	USkeletalMesh* InSkeletalMesh,
-	IMeshUtilities& InMeshUtilities
-	)
+	IMeshUtilities& InMeshUtilities,
+	const bool bCacheOptimize = true)
 {
 	FSkeletalMeshModel* ImportedModels = InSkeletalMesh->GetImportedModel();
 	const int32 LODIndex = ImportedModels->LODModels.Num(); 
@@ -220,6 +220,7 @@ static bool AddLODFromMeshDescription(
 	IMeshUtilities::MeshBuildOptions BuildOptions;
 	BuildOptions.TargetPlatform = GetTargetPlatformManagerRef().GetRunningTargetPlatform();
 	BuildOptions.FillOptions(InSkeletalMesh->GetLODInfo(InSkeletalMesh->GetLODNum() - 1)->BuildSettings);
+	BuildOptions.bCacheOptimize = bCacheOptimize;
 
 	TArray<FText> WarningMessages;
 	if (!InMeshUtilities.BuildSkeletalMesh(SkeletalMeshModel, InSkeletalMesh->GetPathName(), InSkeletalMesh->GetRefSkeleton(), LODInfluences, LODWedges, LODFaces, LODPoints, LODPointToRawMap, BuildOptions, &WarningMessages, nullptr))
@@ -502,7 +503,8 @@ bool FStaticToSkeletalMeshConverter::InitializeSkeletalMeshFromMeshDescriptions(
 	TConstArrayView<FSkeletalMaterial> InMaterials,
 	const FReferenceSkeleton& InReferenceSkeleton,
 	const bool bInRecomputeNormals,
-	const bool bInRecomputeTangents
+	const bool bInRecomputeTangents,
+	const bool bCacheOptimize
 )	
 {
 	if (!ensure(InSkeletalMesh))
@@ -586,7 +588,7 @@ bool FStaticToSkeletalMeshConverter::InitializeSkeletalMeshFromMeshDescriptions(
 				}
 			}
 			
-			if (!AddLODFromMeshDescription(MoveTemp(ClonedDescription), InSkeletalMesh, MeshUtilities))
+			if (!AddLODFromMeshDescription(MoveTemp(ClonedDescription), InSkeletalMesh, MeshUtilities, bCacheOptimize))
 			{
 				// If we didn't get a model for LOD index 0, we don't have a mesh. Bail out.
 				if (bFirstSourceModel)
diff --git a/Engine/Source/Developer/SkeletalMeshUtilitiesCommon/Public/StaticToSkeletalMeshConverter.h b/Engine/Source/Developer/SkeletalMeshUtilitiesCommon/Public/StaticToSkeletalMeshConverter.h
--- a/Engine/Source/Developer/SkeletalMeshUtilitiesCommon/Public/StaticToSkeletalMeshConverter.h
+++ b/Engine/Source/Developer/SkeletalMeshUtilitiesCommon/Public/StaticToSkeletalMeshConverter.h
@@ -79,6 +79,7 @@ struct SKELETALMESHUTILITIESCOMMON_API FStaticToSkeletalMeshConverter
 	 *  @param InReferenceSkeleton The reference skeleton to use.
 	 *  @param bInRecomputeNormals Recompute normals when the mesh is generated from the mesh description.
 	 *  @param bInRecomputeTangents Recompute tangents when the mesh is generated from the mesh description.
+	 *  @param bCacheOptimize Enable the cache optimization of the index buffer.
 	 */
 	static bool InitializeSkeletalMeshFromMeshDescriptions(
 		USkeletalMesh* InSkeletalMesh,
@@ -86,7 +87,8 @@ struct SKELETALMESHUTILITIESCOMMON_API FStaticToSkeletalMeshConverter
 		TConstArrayView<FSkeletalMaterial> InMaterials,
 		const FReferenceSkeleton& InReferenceSkeleton,
 		const bool bInRecomputeNormals = false,
-		const bool bInRecomputeTangents = true
+		const bool bInRecomputeTangents = true,
+		const bool bCacheOptimize = true
 		);
 };
 
