Updated headers using UnrealCodeFixup to make sure dllstorage is on methods/staticvars instead of types. Part 4/n

diff --git a/Engine/Source/Developer/SkeletalMeshUtilitiesCommon/Public/LODUtilities.h b/Engine/Source/Developer/SkeletalMeshUtilitiesCommon/Public/LODUtilities.h
--- a/Engine/Source/Developer/SkeletalMeshUtilitiesCommon/Public/LODUtilities.h
+++ b/Engine/Source/Developer/SkeletalMeshUtilitiesCommon/Public/LODUtilities.h
@@ -45,7 +45,7 @@ class FSkeletalMeshImportData;
 // FLODUtilities
 
 
-class SKELETALMESHUTILITIESCOMMON_API FLODUtilities
+class FLODUtilities
 {
 public:
 
@@ -56,7 +56,7 @@ public:
 	* @param Influences - BoneWeights and Ids for the corresponding vertices. 
 	* @param MeshName	- Name of mesh, used for warning messages
 	*/
-	static void ProcessImportMeshInfluences(const int32 VertexCount, TArray<SkeletalMeshImportData::FRawBoneInfluence>& Influences, const FString& MeshName);
+	static SKELETALMESHUTILITIESCOMMON_API void ProcessImportMeshInfluences(const int32 VertexCount, TArray<SkeletalMeshImportData::FRawBoneInfluence>& Influences, const FString& MeshName);
 
 	/** Regenerate LODs of the mesh
 	*
@@ -68,21 +68,21 @@ public:
 	*
 	* @return true if succeed. If mesh reduction is not available this will return false.
 	*/
-	static bool RegenerateLOD(USkeletalMesh* SkeletalMesh, const class ITargetPlatform* TargetPlatform, int32 NewLODCount = 0, bool bRegenerateEvenIfImported = false, bool bGenerateBaseLOD = false);
+	static SKELETALMESHUTILITIESCOMMON_API bool RegenerateLOD(USkeletalMesh* SkeletalMesh, const class ITargetPlatform* TargetPlatform, int32 NewLODCount = 0, bool bRegenerateEvenIfImported = false, bool bGenerateBaseLOD = false);
 
 	/** Removes a particular LOD from the SkeletalMesh. 
 	*
 	* @param UpdateContext - The skeletal mesh and actor components to operate on.
 	* @param DesiredLOD   - The LOD index to remove the LOD from.
 	*/
-	static void RemoveLOD( FSkeletalMeshUpdateContext& UpdateContext, int32 DesiredLOD );
+	static SKELETALMESHUTILITIESCOMMON_API void RemoveLOD( FSkeletalMeshUpdateContext& UpdateContext, int32 DesiredLOD );
 
 	/** Removes the specified LODs from the SkeletalMesh.
 	*
 	* @param UpdateContext - The skeletal mesh and actor components to operate on.
 	* @param DesiredLODs   - The array of LOD index to remove the LOD from. The order is irrelevant since the array will be sorted to be reverse iterate.
 	*/
-	static void RemoveLODs(FSkeletalMeshUpdateContext& UpdateContext, const TArray<int32>& DesiredLODs);
+	static SKELETALMESHUTILITIESCOMMON_API void RemoveLODs(FSkeletalMeshUpdateContext& UpdateContext, const TArray<int32>& DesiredLODs);
 
 	/*
 	 * Add or change the LOD data specified by LodIndex with the content of the sourceSkeletalMesh.
@@ -92,7 +92,7 @@ public:
 	 * @Param LodIndex - The destination lod index we want to add or modify.
 	 * @Param SourceDataFilename - The file name we use to import the source skeletal mesh.
 	 */
-	static bool SetCustomLOD(USkeletalMesh* DestinationSkeletalMesh, USkeletalMesh* SourceSkeletalMesh, const int32 LodIndex, const FString& SourceDataFilename);
+	static SKELETALMESHUTILITIESCOMMON_API bool SetCustomLOD(USkeletalMesh* DestinationSkeletalMesh, USkeletalMesh* SourceSkeletalMesh, const int32 LodIndex, const FString& SourceDataFilename);
 
 	/**
 	*	Simplifies the static mesh based upon various user settings for DesiredLOD.
@@ -100,7 +100,7 @@ public:
 	* @param UpdateContext - The skeletal mesh and actor components to operate on.
 	* @param DesiredLOD - The LOD to simplify
 	*/
-	static void SimplifySkeletalMeshLOD(FSkeletalMeshUpdateContext& UpdateContext, int32 DesiredLOD, const class ITargetPlatform* TargetPlatform, bool bRestoreClothing = false, std::atomic<bool>* OutNeedsPackageDirtied = nullptr);
+	static SKELETALMESHUTILITIESCOMMON_API void SimplifySkeletalMeshLOD(FSkeletalMeshUpdateContext& UpdateContext, int32 DesiredLOD, const class ITargetPlatform* TargetPlatform, bool bRestoreClothing = false, std::atomic<bool>* OutNeedsPackageDirtied = nullptr);
 
 	/**
 	*	Restore the LOD imported model to the last imported data. Call this function if you want to remove the reduce on the base LOD
@@ -109,7 +109,7 @@ public:
 	* @param LodIndex - The LOD index to restore the imported LOD model
 	* @param bReregisterComponent - if true the component using the skeletal mesh will all be re register.
 	*/
-	static bool RestoreSkeletalMeshLODImportedData_DEPRECATED(USkeletalMesh* SkeletalMesh, int32 LodIndex);
+	static SKELETALMESHUTILITIESCOMMON_API bool RestoreSkeletalMeshLODImportedData_DEPRECATED(USkeletalMesh* SkeletalMesh, int32 LodIndex);
 	
 	/**
 	 * Refresh LOD Change
@@ -119,23 +119,23 @@ public:
 	 *
 	 * @param	InSkeletalMesh	SkeletalMesh that LOD has been changed for
 	 */
-	static void RefreshLODChange(const USkeletalMesh* SkeletalMesh);
+	static SKELETALMESHUTILITIESCOMMON_API void RefreshLODChange(const USkeletalMesh* SkeletalMesh);
 
 
 	/*
 	 * Regenerate LOD that are dependent of LODIndex
 	 */
-	static void RegenerateDependentLODs(USkeletalMesh* SkeletalMesh, int32 LODIndex, const class ITargetPlatform* TargetPlatform);
+	static SKELETALMESHUTILITIESCOMMON_API void RegenerateDependentLODs(USkeletalMesh* SkeletalMesh, int32 LODIndex, const class ITargetPlatform* TargetPlatform);
 
 	/*
 	 * Build the morph targets for the specified LOD. The function use the Morph target data stored in the FSkeletalMeshImportData ImportData structure
 	 */
-	static void BuildMorphTargets(USkeletalMesh* SkeletalMesh, class FSkeletalMeshImportData &ImportData, int32 LODIndex, bool ShouldImportNormals, bool ShouldImportTangents, bool bUseMikkTSpace, const FOverlappingThresholds& Thresholds);
+	static SKELETALMESHUTILITIESCOMMON_API void BuildMorphTargets(USkeletalMesh* SkeletalMesh, class FSkeletalMeshImportData &ImportData, int32 LODIndex, bool ShouldImportNormals, bool ShouldImportTangents, bool bUseMikkTSpace, const FOverlappingThresholds& Thresholds);
 
 	/*
 	 * Same as above but use normals from the source mesh description to build up the morph targets. 
 	 */
-	static void BuildMorphTargets(
+	static SKELETALMESHUTILITIESCOMMON_API void BuildMorphTargets(
 		USkeletalMesh* SkeletalMesh,
 		const FMeshDescription& SkeletalMeshModel,
 		FSkeletalMeshImportData& ImportData,
@@ -158,7 +158,7 @@ public:
 	 * @param LODIndexDest - the destination LOD
 	 * @param LODIndexSrc - the Source LOD index
 	 */
-	static bool UpdateAlternateSkinWeights(
+	static SKELETALMESHUTILITIESCOMMON_API bool UpdateAlternateSkinWeights(
 		USkeletalMesh* SkeletalMeshDest,
 		const FName& ProfileNameDest,
 		USkeletalMesh* SkeletalMeshSrc,
@@ -167,7 +167,7 @@ public:
 		const IMeshUtilities::MeshBuildOptions& Options);
 
 	UE_DEPRECATED(5.2, "Please use the new overloads of UpdateAlternateSkinWeights that take an IMeshUtilities::MeshBuildOptions. Note that IMeshUtilities::MeshBuildOptions::bComputeNormals/Tangents has the opposite meaning of ShouldImportNormals/Tangents.")
-	static bool UpdateAlternateSkinWeights(
+	static SKELETALMESHUTILITIESCOMMON_API bool UpdateAlternateSkinWeights(
 		USkeletalMesh* SkeletalMeshDest,
 		const FName& ProfileNameDest,
 		USkeletalMesh* SkeletalMeshSrc,
@@ -186,14 +186,14 @@ public:
 	 * @param SkeletalMeshDest - The skeletal mesh that will receive the alternate skinning weights.
 	 * @param LODIndexDest - the destination LOD
 	 */
-	static bool UpdateAlternateSkinWeights(
+	static SKELETALMESHUTILITIESCOMMON_API bool UpdateAlternateSkinWeights(
 		USkeletalMesh* SkeletalMeshDest,
 		const FName& ProfileNameDest,
 		int32 LODIndexDest,
 		const IMeshUtilities::MeshBuildOptions& Options);
 
 	UE_DEPRECATED(5.2, "Please use the new overloads of UpdateAlternateSkinWeights that take an IMeshUtilities::MeshBuildOptions. Note that IMeshUtilities::MeshBuildOptions::bComputeNormals/Tangents has the opposite meaning of ShouldImportNormals/Tangents.")
-	static bool UpdateAlternateSkinWeights(
+	static SKELETALMESHUTILITIESCOMMON_API bool UpdateAlternateSkinWeights(
 		USkeletalMesh* SkeletalMeshDest,
 		const FName& ProfileNameDest,
 		int32 LODIndexDest,
@@ -203,7 +203,7 @@ public:
 		bool bUseMikkTSpace,
 		bool bComputeWeightedNormals);
 
-	static bool UpdateAlternateSkinWeights(
+	static SKELETALMESHUTILITIESCOMMON_API bool UpdateAlternateSkinWeights(
 		FSkeletalMeshLODModel& LODModelDest,
 		FSkeletalMeshImportData& ImportDataDest,
 		USkeletalMesh* SkeletalMeshDest,
@@ -213,7 +213,7 @@ public:
 		const IMeshUtilities::MeshBuildOptions& Options);
 
 	UE_DEPRECATED(5.2, "Please use the new overloads of UpdateAlternateSkinWeights that take an IMeshUtilities::MeshBuildOptions. Note that IMeshUtilities::MeshBuildOptions::bComputeNormals/Tangents has the opposite meaning of ShouldImportNormals/Tangents.")
-	static bool UpdateAlternateSkinWeights(
+	static SKELETALMESHUTILITIESCOMMON_API bool UpdateAlternateSkinWeights(
 		FSkeletalMeshLODModel& LODModelDest,
 		FSkeletalMeshImportData& ImportDataDest,
 		USkeletalMesh* SkeletalMeshDest,
@@ -228,20 +228,20 @@ public:
 
 	
 	/** Build the vertex attributes */
-	static bool UpdateLODInfoVertexAttributes(USkeletalMesh *InSkeletalMesh, int32 InSourceLODIndex, int32 InTargetLODIndex, bool bInCopyAttributeValues);
+	static SKELETALMESHUTILITIESCOMMON_API bool UpdateLODInfoVertexAttributes(USkeletalMesh *InSkeletalMesh, int32 InSourceLODIndex, int32 InTargetLODIndex, bool bInCopyAttributeValues);
 	
 	/**
 	 * Re-generate all (editor-only) skin weight profile, used whenever we rebuild the skeletal mesh data which could change the chunking and bone indices
 	 * 
 	 * If BoneInfluenceLimit is 0, the DefaultBoneInfluenceLimit from the project settings will be used.
 	 */
-	static void RegenerateAllImportSkinWeightProfileData(FSkeletalMeshLODModel& LODModelDest, int32 BoneInfluenceLimit = 0, const ITargetPlatform* TargetPlatform = nullptr);
+	static SKELETALMESHUTILITIESCOMMON_API void RegenerateAllImportSkinWeightProfileData(FSkeletalMeshLODModel& LODModelDest, int32 BoneInfluenceLimit = 0, const ITargetPlatform* TargetPlatform = nullptr);
 
-	static void UnbindClothingAndBackup(USkeletalMesh* SkeletalMesh, TArray<ClothingAssetUtils::FClothingAssetMeshBinding>& ClothingBindings);
-	static void UnbindClothingAndBackup(USkeletalMesh* SkeletalMesh, TArray<ClothingAssetUtils::FClothingAssetMeshBinding>& ClothingBindings, const int32 LODIndex);
+	static SKELETALMESHUTILITIESCOMMON_API void UnbindClothingAndBackup(USkeletalMesh* SkeletalMesh, TArray<ClothingAssetUtils::FClothingAssetMeshBinding>& ClothingBindings);
+	static SKELETALMESHUTILITIESCOMMON_API void UnbindClothingAndBackup(USkeletalMesh* SkeletalMesh, TArray<ClothingAssetUtils::FClothingAssetMeshBinding>& ClothingBindings, const int32 LODIndex);
 	
-	static void RestoreClothingFromBackup(USkeletalMesh* SkeletalMesh, TArray<ClothingAssetUtils::FClothingAssetMeshBinding>& ClothingBindings);
-	static void RestoreClothingFromBackup(USkeletalMesh* SkeletalMesh, TArray<ClothingAssetUtils::FClothingAssetMeshBinding>& ClothingBindings, const int32 LODIndex);
+	static SKELETALMESHUTILITIESCOMMON_API void RestoreClothingFromBackup(USkeletalMesh* SkeletalMesh, TArray<ClothingAssetUtils::FClothingAssetMeshBinding>& ClothingBindings);
+	static SKELETALMESHUTILITIESCOMMON_API void RestoreClothingFromBackup(USkeletalMesh* SkeletalMesh, TArray<ClothingAssetUtils::FClothingAssetMeshBinding>& ClothingBindings, const int32 LODIndex);
 
 	struct FMorphTargetLodBackupData
 	{
@@ -253,12 +253,12 @@ public:
 	/**
 	 * Backup all the custom imported morph targets.
 	 */
-	static void BackupCustomImportedMorphTargetData(USkeletalMesh* SkeletalMesh, TMap<FString, TArray<FMorphTargetLodBackupData>>& BackupImportedMorphTargetData);
+	static SKELETALMESHUTILITIESCOMMON_API void BackupCustomImportedMorphTargetData(USkeletalMesh* SkeletalMesh, TMap<FString, TArray<FMorphTargetLodBackupData>>& BackupImportedMorphTargetData);
 
 	/**
 	 * Restore the custom imported morph target in the mesh description. Return true if the mesh description was changed because it restore some morph targets
 	 */
-	static bool RestoreCustomImportedMorphTargetData(USkeletalMesh* SkeletalMesh, const int32 LodIndex, FMeshDescription& LodMeshDescription, const TMap<FString, TArray<FMorphTargetLodBackupData>>& BackupImportedMorphTargetData);
+	static SKELETALMESHUTILITIESCOMMON_API bool RestoreCustomImportedMorphTargetData(USkeletalMesh* SkeletalMesh, const int32 LodIndex, FMeshDescription& LodMeshDescription, const TMap<FString, TArray<FMorphTargetLodBackupData>>& BackupImportedMorphTargetData);
 
 	/**
 	 * Before building skeletalmesh base LOD (LOD index 0) using MeshUtilities.BuildSkeletalMesh, we want to adjust the imported faces material index to point on the correct sk material. We use the material name to match the material.
@@ -267,7 +267,7 @@ public:
 	 * @param LODFaces - The face to adjust with the material remap
 	 * @param LODIndex - We can adjust only the base LOD (LODIndex 0), the function will not do anything if LODIndex is not 0
 	 */
-	static void AdjustImportDataFaceMaterialIndex(const TArray<FSkeletalMaterial>& Materials, TArray<SkeletalMeshImportData::FMaterial>& RawMeshMaterials, TArray<SkeletalMeshImportData::FMeshFace>& LODFaces, int32 LODIndex);
+	static SKELETALMESHUTILITIESCOMMON_API void AdjustImportDataFaceMaterialIndex(const TArray<FSkeletalMaterial>& Materials, TArray<SkeletalMeshImportData::FMaterial>& RawMeshMaterials, TArray<SkeletalMeshImportData::FMeshFace>& LODFaces, int32 LODIndex);
 
 	/**
 	 * Structure to pass all the needed parameters to do match the material when importing a skeletal mesh LOD.
@@ -293,17 +293,17 @@ public:
 	 * In case we re-import a LOD this code will match the sections using imported material name to be able to keep any section data.
 	 * The material slot array is not re-order even if we re-import the base LOD.
 	 */
-	static void MatchImportedMaterials(FSkeletalMeshMatchImportedMaterialsParameters& Parameters);
+	static SKELETALMESHUTILITIESCOMMON_API void MatchImportedMaterials(FSkeletalMeshMatchImportedMaterialsParameters& Parameters);
 
 	/**
 	 * Reorder the material slot array to follow the base LOD section order. It will readjust all LOD section material index and LODMaterialMap.
 	 */
-	static void ReorderMaterialSlotToBaseLod(USkeletalMesh* SkeletalMesh);
+	static SKELETALMESHUTILITIESCOMMON_API void ReorderMaterialSlotToBaseLod(USkeletalMesh* SkeletalMesh);
 
 	/**
 	 * Remove any material slot that is not used by any LODs
 	 */
-	static void RemoveUnusedMaterialSlot(USkeletalMesh* SkeletalMesh);
+	static SKELETALMESHUTILITIESCOMMON_API void RemoveUnusedMaterialSlot(USkeletalMesh* SkeletalMesh);
 
 	/**
 	 * This function will strip all triangle in the specified LOD that don't have any UV area pointing on a black pixel in the TextureMask.
@@ -314,13 +314,13 @@ public:
 	 * @Param TextureMask: The texture containing the stripping mask. non black pixel strip triangle, black pixel keep them.
 	 * @Param Threshold: The threshold we want when comparing the texture value with zero.
 	 */
-	static bool StripLODGeometry(USkeletalMesh* SkeletalMesh, const int32 LODIndex, UTexture2D* TextureMask, const float Threshold);
+	static SKELETALMESHUTILITIESCOMMON_API bool StripLODGeometry(USkeletalMesh* SkeletalMesh, const int32 LODIndex, UTexture2D* TextureMask, const float Threshold);
 
 private:
 	FLODUtilities() {}
 
 	/** Generate the editor-only data stored for a skin weight profile (relies on bone indices) */
-	static void GenerateImportedSkinWeightProfileData(
+	static SKELETALMESHUTILITIESCOMMON_API void GenerateImportedSkinWeightProfileData(
 		FSkeletalMeshLODModel& LODModelDest,
 		struct FImportedSkinWeightProfileData& ImportedProfileData,
 		int32 BoneInfluenceLimit,
@@ -366,7 +366,7 @@ private:
 	 * @param SkeletalMesh - The skeletal mesh and actor components to operate on.
 	 * @param DesiredLOD - Desired LOD
 	 */
-	static void SimplifySkeletalMeshLOD(
+	static SKELETALMESHUTILITIESCOMMON_API void SimplifySkeletalMeshLOD(
 		USkeletalMesh* SkeletalMesh, 
 		int32 DesiredLOD, 
 		const ITargetPlatform* TargetPlatform,
@@ -382,7 +382,7 @@ private:
 	* @param SourceLOD      - The source LOD morph target .
 	* @param DestinationLOD   - The destination LOD morph target to apply the source LOD morph target
 	*/
-	static void ApplyMorphTargetsToLOD(USkeletalMesh* SkeletalMesh, int32 SourceLOD, int32 DestinationLOD, const FInlineReductionDataParameter& InlineApplyMorphTargetParameter);
+	static SKELETALMESHUTILITIESCOMMON_API void ApplyMorphTargetsToLOD(USkeletalMesh* SkeletalMesh, int32 SourceLOD, int32 DestinationLOD, const FInlineReductionDataParameter& InlineApplyMorphTargetParameter);
 
 	/**
 	*  Clear generated morphtargets for the given LODs
@@ -390,7 +390,7 @@ private:
 	* @param SkeletalMesh - The skeletal mesh and actor components to operate on.
 	* @param DesiredLOD - Desired LOD
 	*/
-	static void ClearGeneratedMorphTarget(USkeletalMesh* SkeletalMesh, int32 DesiredLOD);
+	static SKELETALMESHUTILITIESCOMMON_API void ClearGeneratedMorphTarget(USkeletalMesh* SkeletalMesh, int32 DesiredLOD);
 };
 
-#endif //WITH_EDITOR
\ No newline at end of file
+#endif //WITH_EDITOR
diff --git a/Engine/Source/Developer/SkeletalMeshUtilitiesCommon/Public/SkeletalMeshNotifier.h b/Engine/Source/Developer/SkeletalMeshUtilitiesCommon/Public/SkeletalMeshNotifier.h
--- a/Engine/Source/Developer/SkeletalMeshUtilitiesCommon/Public/SkeletalMeshNotifier.h
+++ b/Engine/Source/Developer/SkeletalMeshUtilitiesCommon/Public/SkeletalMeshNotifier.h
@@ -6,6 +6,8 @@
 #include "Delegates/Delegate.h"
 #include "Containers/Array.h"
 
+#define UE_API SKELETALMESHUTILITIESCOMMON_API
+
 class FName;
 class HHitProxy;
 
@@ -22,26 +24,26 @@ enum class ESkeletalMeshNotifyType
 // A delegate for monitoring to skeletal mesh global notifications.
 DECLARE_MULTICAST_DELEGATE_TwoParams(FSkeletalMeshNotifyDelegate, const TArray<FName>& /*InBoneNames*/, const ESkeletalMeshNotifyType /*InNotifType*/);
 
-class SKELETALMESHUTILITIESCOMMON_API ISkeletalMeshNotifier
+class ISkeletalMeshNotifier
 {
 public:
 	ISkeletalMeshNotifier() = default;
 	virtual ~ISkeletalMeshNotifier() = default;
 
-	FSkeletalMeshNotifyDelegate& Delegate();
+	UE_API FSkeletalMeshNotifyDelegate& Delegate();
 
 	// override this function to react to notifications locally.
 	virtual void HandleNotification(const TArray<FName>& BoneNames, const ESkeletalMeshNotifyType InNotifyType) = 0;
 
-	void Notify(const TArray<FName>& BoneNames, const ESkeletalMeshNotifyType InNotifyType) const;
-	bool Notifying() const;
+	UE_API void Notify(const TArray<FName>& BoneNames, const ESkeletalMeshNotifyType InNotifyType) const;
+	UE_API bool Notifying() const;
 
 private:
 	FSkeletalMeshNotifyDelegate NotifyDelegate;
 	mutable bool bNotifying = false;
 };
 
-class SKELETALMESHUTILITIESCOMMON_API ISkeletalMeshEditorBinding
+class ISkeletalMeshEditorBinding
 {
 public:
 	ISkeletalMeshEditorBinding() = default;
@@ -56,4 +58,6 @@ public:
 	virtual TArray<FName> GetSelectedMorphTargets() const = 0;
 	
 
-};
\ No newline at end of file
+};
+
+#undef UE_API
diff --git a/Engine/Source/Developer/SkeletalMeshUtilitiesCommon/Public/StaticToSkeletalMeshConverter.h b/Engine/Source/Developer/SkeletalMeshUtilitiesCommon/Public/StaticToSkeletalMeshConverter.h
--- a/Engine/Source/Developer/SkeletalMeshUtilitiesCommon/Public/StaticToSkeletalMeshConverter.h
+++ b/Engine/Source/Developer/SkeletalMeshUtilitiesCommon/Public/StaticToSkeletalMeshConverter.h
@@ -17,7 +17,7 @@ struct FReferenceSkeleton;
 struct FSkeletalMaterial;
 
 
-struct SKELETALMESHUTILITIESCOMMON_API FStaticToSkeletalMeshConverter
+struct FStaticToSkeletalMeshConverter
 {
 	/** Creates a skeleton from a static mesh with a single root bone.
 	 *
@@ -28,7 +28,7 @@ struct SKELETALMESHUTILITIESCOMMON_API FStaticToSkeletalMeshConverter
 	 *    at the center of the bottom of the static mesh's bbox.
 	 *  @return \c true if the operation succeeded. 
 	 */
-	static bool InitializeSkeletonFromStaticMesh(
+	static SKELETALMESHUTILITIESCOMMON_API bool InitializeSkeletonFromStaticMesh(
 		USkeleton* InSkeleton,
 		const UStaticMesh* InStaticMesh,
 		const FVector& InRelativeRootPosition 
@@ -47,7 +47,7 @@ struct SKELETALMESHUTILITIESCOMMON_API FStaticToSkeletalMeshConverter
 	 *  @param InIntermediaryJointCount Number of joints to create between the root and the end effector.  
 	 *  @return \c true if the operation succeeded. 
 	 */
-	static bool InitializeSkeletonFromStaticMesh(
+	static SKELETALMESHUTILITIESCOMMON_API bool InitializeSkeletonFromStaticMesh(
 		USkeleton* InSkeleton,
 		const UStaticMesh* InStaticMesh,
 		const FVector& InRelativeRootPosition,
@@ -62,7 +62,7 @@ struct SKELETALMESHUTILITIESCOMMON_API FStaticToSkeletalMeshConverter
 	 *  @param InReferenceSkeleton The reference skeleton to use.
 	 *  @param InBindBone The bone to bind to. If no bone name is given, the binding defaults to the root bone.
 	 */
-	static bool InitializeSkeletalMeshFromStaticMesh(
+	static SKELETALMESHUTILITIESCOMMON_API bool InitializeSkeletalMeshFromStaticMesh(
 		USkeletalMesh* InSkeletalMesh, 
 		const UStaticMesh* InStaticMesh,
 		const FReferenceSkeleton& InReferenceSkeleton,
@@ -81,7 +81,7 @@ struct SKELETALMESHUTILITIESCOMMON_API FStaticToSkeletalMeshConverter
 	 *  @param bInRecomputeTangents Recompute tangents when the mesh is generated from the mesh description.
 	 *  @param bCacheOptimize Enable the cache optimization of the index buffer.
 	 */
-	static bool InitializeSkeletalMeshFromMeshDescriptions(
+	static SKELETALMESHUTILITIESCOMMON_API bool InitializeSkeletalMeshFromMeshDescriptions(
 		USkeletalMesh* InSkeletalMesh,
 		TArrayView<const FMeshDescription*> InMeshDescriptions,
 		TConstArrayView<FSkeletalMaterial> InMaterials,
diff --git a/Engine/Source/Developer/TraceInsights/Private/Insights/Tests/InsightsTestRunner.h b/Engine/Source/Developer/TraceInsights/Private/Insights/Tests/InsightsTestRunner.h
--- a/Engine/Source/Developer/TraceInsights/Private/Insights/Tests/InsightsTestRunner.h
+++ b/Engine/Source/Developer/TraceInsights/Private/Insights/Tests/InsightsTestRunner.h
@@ -12,33 +12,33 @@
 
 DECLARE_LOG_CATEGORY_EXTERN(InsightsTestRunner, Log, All);
 
-class TRACEINSIGHTS_API FInsightsTestRunner : public TSharedFromThis<FInsightsTestRunner>, public IInsightsComponent
+class FInsightsTestRunner : public TSharedFromThis<FInsightsTestRunner>, public IInsightsComponent
 {
 public:
-	virtual ~FInsightsTestRunner();
+	TRACEINSIGHTS_API virtual ~FInsightsTestRunner();
 
-	void ScheduleCommand(const FString& InCmd);
+	TRACEINSIGHTS_API void ScheduleCommand(const FString& InCmd);
 
-	virtual void Initialize(IUnrealInsightsModule& InsightsModule) override;
-	virtual void Shutdown() override;
-	virtual void RegisterMajorTabs(IUnrealInsightsModule& InsightsModule) override;
-	virtual void UnregisterMajorTabs() override;
+	TRACEINSIGHTS_API virtual void Initialize(IUnrealInsightsModule& InsightsModule) override;
+	TRACEINSIGHTS_API virtual void Shutdown() override;
+	TRACEINSIGHTS_API virtual void RegisterMajorTabs(IUnrealInsightsModule& InsightsModule) override;
+	TRACEINSIGHTS_API virtual void UnregisterMajorTabs() override;
 
-	static TSharedPtr<FInsightsTestRunner> CreateInstance();
-	static TSharedPtr<FInsightsTestRunner> Get();
+	static TRACEINSIGHTS_API TSharedPtr<FInsightsTestRunner> CreateInstance();
+	static TRACEINSIGHTS_API TSharedPtr<FInsightsTestRunner> Get();
 
-	bool Tick(float DeltaTime);
+	TRACEINSIGHTS_API bool Tick(float DeltaTime);
 
 	void SetAutoQuit(bool InAutoQuit) { bAutoQuit = InAutoQuit; }
 	bool GetAutoQuit() const { return bAutoQuit; }
 
 	void SetInitAutomationModules(bool InInitAutomationModules) { bInitAutomationModules = InInitAutomationModules; }
 	bool GetInitAutomationModules() const { return bInitAutomationModules; }
-	void RunTests();
+	TRACEINSIGHTS_API void RunTests();
 
 private:
-	TSharedRef<SDockTab> SpawnAutomationWindowTab(const FSpawnTabArgs& Args);
-	void OnSessionAnalysisCompleted();
+	TRACEINSIGHTS_API TSharedRef<SDockTab> SpawnAutomationWindowTab(const FSpawnTabArgs& Args);
+	TRACEINSIGHTS_API void OnSessionAnalysisCompleted();
 
 private:
 	/** The delegate to be invoked when this manager ticks. */
@@ -56,9 +56,9 @@ private:
 	bool bIsRunningTests = false;
 	bool bIsAnalysisComplete = false;
 
-	static const TCHAR* AutoQuitMsgOnComplete;
+	static TRACEINSIGHTS_API const TCHAR* AutoQuitMsgOnComplete;
 
-	static TSharedPtr<FInsightsTestRunner> Instance;
+	static TRACEINSIGHTS_API TSharedPtr<FInsightsTestRunner> Instance;
 };
 
 #endif //UE_BUILD_SHIPPING && !WITH_EDITOR
diff --git a/Engine/Source/Developer/TraceInsights/Public/Insights/IInsightsManager.h b/Engine/Source/Developer/TraceInsights/Public/Insights/IInsightsManager.h
--- a/Engine/Source/Developer/TraceInsights/Public/Insights/IInsightsManager.h
+++ b/Engine/Source/Developer/TraceInsights/Public/Insights/IInsightsManager.h
@@ -10,7 +10,7 @@ namespace Insights
 ////////////////////////////////////////////////////////////////////////////////////////////////////
 
 /** Interface for an Insights manager. */
-class TRACEINSIGHTS_API IInsightsManager : public IInsightsComponent
+class IInsightsManager : public IInsightsComponent
 {
 public:
 	/** The event to execute when the session has changed. */
diff --git a/Engine/Source/Developer/TraceInsights/Public/Insights/Tests/InsightsTestUtils.h b/Engine/Source/Developer/TraceInsights/Public/Insights/Tests/InsightsTestUtils.h
--- a/Engine/Source/Developer/TraceInsights/Public/Insights/Tests/InsightsTestUtils.h
+++ b/Engine/Source/Developer/TraceInsights/Public/Insights/Tests/InsightsTestUtils.h
@@ -7,24 +7,28 @@
 #include "Containers/UnrealString.h"
 #include "ProfilingDebugging/TraceAuxiliary.h" // for FTraceAuxiliary::EConnectionType
 
+#define UE_API TRACEINSIGHTS_API
+
 class FAutomationTestBase;
 
-class TRACEINSIGHTS_API FInsightsTestUtils
+class FInsightsTestUtils
 {
 public:
-	FInsightsTestUtils(FAutomationTestBase* Test);
-
-	bool AnalyzeTrace(const TCHAR* Path) const;
-	bool FileContainsString(const FString& PathToFile, const FString& ExpectedString, double Timeout) const;
-	bool SetupUTS(double Timeout = 30.0, bool bUseFork = false) const;
-	bool KillUTS(double Timeout = 30.0) const;
-	bool StartTracing(FTraceAuxiliary::EConnectionType ConnectionType, double Timeout = 10.0) const;
-	bool IsUnrealTraceServerReady(const TCHAR* Host = TEXT("localhost"), int32 Port = 0U) const;
-	bool IsTraceHasLiveStatus(const FString& TraceName, const TCHAR* Host = TEXT("localhost"), int32 Port = 0U) const;
-	void ResetSession() const;
-	FString GetLiveTrace(const TCHAR* Host = TEXT("localhost"), int32 Port = 0U) const;
-	FString GetUTSPath() const;
+	UE_API FInsightsTestUtils(FAutomationTestBase* Test);
+
+	UE_API bool AnalyzeTrace(const TCHAR* Path) const;
+	UE_API bool FileContainsString(const FString& PathToFile, const FString& ExpectedString, double Timeout) const;
+	UE_API bool SetupUTS(double Timeout = 30.0, bool bUseFork = false) const;
+	UE_API bool KillUTS(double Timeout = 30.0) const;
+	UE_API bool StartTracing(FTraceAuxiliary::EConnectionType ConnectionType, double Timeout = 10.0) const;
+	UE_API bool IsUnrealTraceServerReady(const TCHAR* Host = TEXT("localhost"), int32 Port = 0U) const;
+	UE_API bool IsTraceHasLiveStatus(const FString& TraceName, const TCHAR* Host = TEXT("localhost"), int32 Port = 0U) const;
+	UE_API void ResetSession() const;
+	UE_API FString GetLiveTrace(const TCHAR* Host = TEXT("localhost"), int32 Port = 0U) const;
+	UE_API FString GetUTSPath() const;
 
 private:
 	FAutomationTestBase* Test;
 };
+
+#undef UE_API
diff --git a/Engine/Source/Developer/TraceInsights/Public/Insights/Tests/TimingProfilerTests.h b/Engine/Source/Developer/TraceInsights/Public/Insights/Tests/TimingProfilerTests.h
--- a/Engine/Source/Developer/TraceInsights/Public/Insights/Tests/TimingProfilerTests.h
+++ b/Engine/Source/Developer/TraceInsights/Public/Insights/Tests/TimingProfilerTests.h
@@ -9,6 +9,8 @@
 // TraceServices
 #include "TraceServices/Model/TimingProfiler.h" // for TraceServices::EEventSortOrder
 
+#define UE_API TRACEINSIGHTS_API
+
 DECLARE_LOG_CATEGORY_EXTERN(TimingProfilerTests, Log, All);
 
 class FAutomationTestBase;
@@ -18,7 +20,7 @@ class FAutomationTestBase;
  * A class containing code for parametric tests for Insight functionality
  * Intended to be called from automatic or user triggered tests
  */
-class TRACEINSIGHTS_API FTimingProfilerTests
+class FTimingProfilerTests
 {
 public:
 	struct FCheckValues
@@ -38,12 +40,14 @@ public:
 		TraceServices::EEventSortOrder SortOrder = TraceServices::EEventSortOrder::ByEndTime;
 	};
 
-	static void RunEnumerateBenchmark(const FEnumerateTestParams& InParams, FCheckValues& OutCheckValues);
-	static void RunEnumerateAsyncBenchmark(const FEnumerateTestParams& InParams, FCheckValues& OutCheckValues);
-	static void RunEnumerateAllTracksBenchmark(const FEnumerateTestParams& InParams, FCheckValues& OutCheckValues);
-	static void RunEnumerateAsyncAllTracksBenchmark(const FEnumerateTestParams& InParams, FCheckValues& OutCheckValues);
-	static bool RunEnumerateSyncAsyncComparisonTest(FAutomationTestBase& Test, const FEnumerateTestParams& InParam, bool bGameThreadOnly);
+	static UE_API void RunEnumerateBenchmark(const FEnumerateTestParams& InParams, FCheckValues& OutCheckValues);
+	static UE_API void RunEnumerateAsyncBenchmark(const FEnumerateTestParams& InParams, FCheckValues& OutCheckValues);
+	static UE_API void RunEnumerateAllTracksBenchmark(const FEnumerateTestParams& InParams, FCheckValues& OutCheckValues);
+	static UE_API void RunEnumerateAsyncAllTracksBenchmark(const FEnumerateTestParams& InParams, FCheckValues& OutCheckValues);
+	static UE_API bool RunEnumerateSyncAsyncComparisonTest(FAutomationTestBase& Test, const FEnumerateTestParams& InParam, bool bGameThreadOnly);
 
-	static uint32 GetTimelineIndex(const TCHAR* InName);
-	static void VerifyCheckValues(FAutomationTestBase& Test, FCheckValues First, FCheckValues Second);
+	static UE_API uint32 GetTimelineIndex(const TCHAR* InName);
+	static UE_API void VerifyCheckValues(FAutomationTestBase& Test, FCheckValues First, FCheckValues Second);
 };
+
+#undef UE_API
diff --git a/Engine/Source/Developer/TraceInsights/Public/Insights/ViewModels/AxisViewportDouble.h b/Engine/Source/Developer/TraceInsights/Public/Insights/ViewModels/AxisViewportDouble.h
--- a/Engine/Source/Developer/TraceInsights/Public/Insights/ViewModels/AxisViewportDouble.h
+++ b/Engine/Source/Developer/TraceInsights/Public/Insights/ViewModels/AxisViewportDouble.h
@@ -4,11 +4,13 @@
 
 #include "CoreMinimal.h"
 
+#define UE_API TRACEINSIGHTS_API
+
 class SScrollBar;
 
 ////////////////////////////////////////////////////////////////////////////////////////////////////
 
-class TRACEINSIGHTS_API FAxisViewportDouble
+class FAxisViewportDouble
 {
 private:
 	static constexpr float SLATE_UNITS_TOLERANCE = 0.1f;
@@ -45,7 +47,7 @@ public:
 	 * The viewport's width if this is a horizontal axis or the viewport's height if this is a vertical axis.
 	 */
 	float GetSize() const { return Size; }
-	bool SetSize(const float InSize);
+	UE_API bool SetSize(const float InSize);
 
 	double GetMinValue() const { return MinValue; }
 	double GetMaxValue() const { return MaxValue; }
@@ -132,16 +134,16 @@ public:
 		MaxScale = InMaxScale;
 	}
 
-	bool SetScale(const double NewScale);
-	bool ZoomWithFixedOffset(const double NewScale, const float Offset);
-	bool RelativeZoomWithFixedOffset(const float Delta, const float Offset);
+	UE_API bool SetScale(const double NewScale);
+	UE_API bool ZoomWithFixedOffset(const double NewScale, const float Offset);
+	UE_API bool RelativeZoomWithFixedOffset(const float Delta, const float Offset);
 
-	void GetScrollLimits(float& OutMinPos, float& OutMaxPos);
-	bool EnforceScrollLimits(const float InMinPos, const float InMaxPos, const float InterpolationFactor);
-	bool UpdatePosWithinLimits();
+	UE_API void GetScrollLimits(float& OutMinPos, float& OutMaxPos);
+	UE_API bool EnforceScrollLimits(const float InMinPos, const float InMaxPos, const float InterpolationFactor);
+	UE_API bool UpdatePosWithinLimits();
 
-	void OnUserScrolled(TSharedPtr<SScrollBar> ScrollBar, float ScrollOffset);
-	void UpdateScrollBar(TSharedPtr<SScrollBar> ScrollBar) const;
+	UE_API void OnUserScrolled(TSharedPtr<SScrollBar> ScrollBar, float ScrollOffset);
+	UE_API void UpdateScrollBar(TSharedPtr<SScrollBar> ScrollBar) const;
 
 private:
 	void OnSizeChanged()
@@ -189,3 +191,5 @@ private:
 };
 
 ////////////////////////////////////////////////////////////////////////////////////////////////////
+
+#undef UE_API
diff --git a/Engine/Source/Developer/TraceInsights/Public/Insights/ViewModels/AxisViewportInt32.h b/Engine/Source/Developer/TraceInsights/Public/Insights/ViewModels/AxisViewportInt32.h
--- a/Engine/Source/Developer/TraceInsights/Public/Insights/ViewModels/AxisViewportInt32.h
+++ b/Engine/Source/Developer/TraceInsights/Public/Insights/ViewModels/AxisViewportInt32.h
@@ -4,11 +4,13 @@
 
 #include "CoreMinimal.h"
 
+#define UE_API TRACEINSIGHTS_API
+
 class SScrollBar;
 
 ////////////////////////////////////////////////////////////////////////////////////////////////////
 
-class TRACEINSIGHTS_API FAxisViewportInt32
+class FAxisViewportInt32
 {
 private:
 	static constexpr float SLATE_UNITS_TOLERANCE = 0.1f;
@@ -54,7 +56,7 @@ public:
 	 * The viewport's width if this is a horizontal axis or the viewport's height if this is a vertical axis.
 	 */
 	float GetSize() const { return Size; }
-	bool SetSize(const float InSize);
+	UE_API bool SetSize(const float InSize);
 
 	int32 GetMinValue() const { return MinValue; }
 	int32 GetMaxValue() const { return MaxValue; }
@@ -150,16 +152,16 @@ public:
 		MaxScale = InMaxScale;
 	}
 
-	bool SetScale(const float NewScale);
-	bool ZoomWithFixedOffset(const float NewScale, const float Offset);
-	bool RelativeZoomWithFixedOffset(const float Delta, const float Offset);
+	UE_API bool SetScale(const float NewScale);
+	UE_API bool ZoomWithFixedOffset(const float NewScale, const float Offset);
+	UE_API bool RelativeZoomWithFixedOffset(const float Delta, const float Offset);
 
-	void GetScrollLimits(float& OutMinPos, float& OutMaxPos);
-	bool EnforceScrollLimits(const float InMinPos, const float InMaxPos, const float InterpolationFactor);
-	bool UpdatePosWithinLimits();
+	UE_API void GetScrollLimits(float& OutMinPos, float& OutMaxPos);
+	UE_API bool EnforceScrollLimits(const float InMinPos, const float InMaxPos, const float InterpolationFactor);
+	UE_API bool UpdatePosWithinLimits();
 
-	void OnUserScrolled(TSharedPtr<SScrollBar> ScrollBar, float ScrollOffset);
-	void UpdateScrollBar(TSharedPtr<SScrollBar> ScrollBar) const;
+	UE_API void OnUserScrolled(TSharedPtr<SScrollBar> ScrollBar, float ScrollOffset);
+	UE_API void UpdateScrollBar(TSharedPtr<SScrollBar> ScrollBar) const;
 
 private:
 	void OnSizeChanged()
@@ -207,3 +209,5 @@ private:
 };
 
 ////////////////////////////////////////////////////////////////////////////////////////////////////
+
+#undef UE_API
diff --git a/Engine/Source/Developer/TraceInsights/Public/Insights/ViewModels/BaseTimingTrack.h b/Engine/Source/Developer/TraceInsights/Public/Insights/ViewModels/BaseTimingTrack.h
--- a/Engine/Source/Developer/TraceInsights/Public/Insights/ViewModels/BaseTimingTrack.h
+++ b/Engine/Source/Developer/TraceInsights/Public/Insights/ViewModels/BaseTimingTrack.h
@@ -130,11 +130,11 @@ public:
 
 ////////////////////////////////////////////////////////////////////////////////////////////////////
 
-class TRACEINSIGHTS_API FBaseTimingTrack : public TSharedFromThis<FBaseTimingTrack>
+class FBaseTimingTrack : public TSharedFromThis<FBaseTimingTrack>
 {
 	friend class FTimingViewDrawHelper;
 
-	INSIGHTS_DECLARE_RTTI_BASE(FBaseTimingTrack)
+	INSIGHTS_DECLARE_RTTI_BASE(FBaseTimingTrack, TRACEINSIGHTS_API)
 
 protected:
 	FBaseTimingTrack()
@@ -150,7 +150,7 @@ protected:
 
 	FBaseTimingTrack(const FBaseTimingTrack&) = delete;
 
-	virtual ~FBaseTimingTrack();
+	TRACEINSIGHTS_API virtual ~FBaseTimingTrack();
 
 public:
 	virtual void Reset()
@@ -167,7 +167,7 @@ public:
 
 	ETimingTrackLocation GetValidLocations() const { return ValidLocations; }
 	ETimingTrackLocation GetLocation() const { return Location; }
-	void SetLocation(ETimingTrackLocation InLocation);
+	TRACEINSIGHTS_API void SetLocation(ETimingTrackLocation InLocation);
 	virtual void OnLocationChanged() { SetDirtyFlag(); }
 
 	int32 GetOrder() const { return Order; }
@@ -238,7 +238,7 @@ public:
 	virtual FReply OnMouseButtonUp(const FGeometry& MyGeometry, const FPointerEvent& MouseEvent) { return FReply::Unhandled(); }
 	virtual FReply OnMouseButtonDoubleClick(const FGeometry& MyGeometry, const FPointerEvent& MouseEvent) { return FReply::Unhandled(); }
 
-	virtual void BuildContextMenu(FMenuBuilder& MenuBuilder);
+	TRACEINSIGHTS_API virtual void BuildContextMenu(FMenuBuilder& MenuBuilder);
 
 	/**
 	 * Gets the event at a specified position. Leaves the InOutEvent unchanged if no event is found at specified position.
@@ -302,20 +302,20 @@ public:
 	 *
 	 * @note This invalidates the result of `GetChildTracks`.
 	 */
-	void AddChildTrack(TSharedRef<FBaseTimingTrack> Track);
+	TRACEINSIGHTS_API void AddChildTrack(TSharedRef<FBaseTimingTrack> Track);
 
 	/**
 	 * Inserts a child track to this track at the specified index.
 	 * @note This invalidates the result of `GetChildTracks`.
 	 */
-	void AddChildTrack(TSharedRef<FBaseTimingTrack> Track, int32 Index);
+	TRACEINSIGHTS_API void AddChildTrack(TSharedRef<FBaseTimingTrack> Track, int32 Index);
 
 	/**
 	 * Removes the provided track from this track's child tracks.
 	 *
 	 * @note This invalidates the result of `GetChildTracks`.
 	 */
-	void RemoveChildTrack(TSharedRef<FBaseTimingTrack> Track);
+	TRACEINSIGHTS_API void RemoveChildTrack(TSharedRef<FBaseTimingTrack> Track);
 
 	// The reason why FindChildTrackOfType does not consider the inheritance hierarchy is because
 	// if it did, downstream inheritors of specific public track types could break plugins defining
@@ -328,7 +328,7 @@ public:
 	 * Note that the type name is matched exactly (without considering inheritance.)
 	 * You should generally prefer the templated version over this.
 	 */
-	TSharedPtr<FBaseTimingTrack> FindChildTrackOfType(FName TrackTypeName);
+	TRACEINSIGHTS_API TSharedPtr<FBaseTimingTrack> FindChildTrackOfType(FName TrackTypeName);
 
 	/**
 	 * Returns the first child track of the given type.
@@ -345,14 +345,14 @@ public:
 	TWeakPtr<FBaseTimingTrack> GetParentTrack() const { return ParentTrack; }
 	bool IsChildTrack() const { return ParentTrack.IsValid(); }
 
-	float GetChildTracksTopHeight(const FTimingViewLayout& Layout) const;
-	void UpdateChildTracksPosY(const FTimingViewLayout& Layout);
+	TRACEINSIGHTS_API float GetChildTracksTopHeight(const FTimingViewLayout& Layout) const;
+	TRACEINSIGHTS_API void UpdateChildTracksPosY(const FTimingViewLayout& Layout);
 
 	// Legacy API supporting a single track. Do not use in new code.
 	UE_DEPRECATED(5.6, "Loop over GetChildTracks instead")
-	TSharedPtr<FBaseTimingTrack> GetChildTrack() const;
+	TRACEINSIGHTS_API TSharedPtr<FBaseTimingTrack> GetChildTrack() const;
 	UE_DEPRECATED(5.6, "Use AddChildTrack/RemoveChildTrack instead")
-	void SetChildTrack(TSharedPtr<FBaseTimingTrack> InTrack);
+	TRACEINSIGHTS_API void SetChildTrack(TSharedPtr<FBaseTimingTrack> InTrack);
 
 	virtual void SetFilterConfigurator(TSharedPtr<UE::Insights::FFilterConfigurator> InFilterConfigurator) {}
 
@@ -383,7 +383,7 @@ private:
 	float Height = 0.0f; // height, in Slate units
 	ETimingTrackFlags Flags = ETimingTrackFlags::IsVisible | ETimingTrackFlags::IsDirty;
 
-	static uint64 IdGenerator;
+	static TRACEINSIGHTS_API uint64 IdGenerator;
 };
 
 ////////////////////////////////////////////////////////////////////////////////////////////////////
diff --git a/Engine/Source/Developer/TraceInsights/Public/Insights/ViewModels/GraphSeries.h b/Engine/Source/Developer/TraceInsights/Public/Insights/ViewModels/GraphSeries.h
--- a/Engine/Source/Developer/TraceInsights/Public/Insights/ViewModels/GraphSeries.h
+++ b/Engine/Source/Developer/TraceInsights/Public/Insights/ViewModels/GraphSeries.h
@@ -21,12 +21,14 @@
 
 #include <limits>
 
+#define UE_API TRACEINSIGHTS_API
+
 class FTimingTrackViewport;
 struct FGraphSeriesEvent;
 
 ////////////////////////////////////////////////////////////////////////////////////////////////////
 
-class TRACEINSIGHTS_API FGraphValueViewport
+class FGraphValueViewport
 {
 public:
 	/**
@@ -72,12 +74,12 @@ private:
 
 ////////////////////////////////////////////////////////////////////////////////////////////////////
 
-class TRACEINSIGHTS_API FGraphSeries
+class FGraphSeries
 {
 	friend class FGraphTrack;
 	friend class FGraphTrackBuilder;
 
-	INSIGHTS_DECLARE_RTTI_BASE(FGraphSeries)
+	INSIGHTS_DECLARE_RTTI_BASE(FGraphSeries, UE_API)
 
 public:
 	struct FBox
@@ -88,8 +90,8 @@ public:
 	};
 
 public:
-	FGraphSeries();
-	virtual ~FGraphSeries();
+	UE_API FGraphSeries();
+	UE_API virtual ~FGraphSeries();
 
 	const FText& GetName() const { return Name; }
 	void SetName(const TCHAR* InName) { Name = FText::FromString(InName); }
@@ -215,15 +217,15 @@ public:
 	 * @return A pointer to an Event located at (X, Y) coordinates, if any; nullptr if no event is located at respective coordinates
 	 * The returned pointer is valid only temporary until next Reset() or next usage of FGraphTrackBuilder for this series/track.
 	 */
-	const FGraphSeriesEvent* GetEvent(const float PosX, const float PosY, const FTimingTrackViewport& Viewport, bool bCheckLine, bool bCheckBox) const;
+	UE_API const FGraphSeriesEvent* GetEvent(const float PosX, const float PosY, const FTimingTrackViewport& Viewport, bool bCheckLine, bool bCheckBox) const;
 
 	/** Updates the track's auto-zoom. Does nothing if IsAutoZoomEnabled() is false. */
-	void UpdateAutoZoom(const float InTopY, const float InBottomY, const double InMinEventValue, const double InMaxEventValue, const bool bIsAutoZoomAnimated = true);
+	UE_API void UpdateAutoZoom(const float InTopY, const float InBottomY, const double InMinEventValue, const double InMaxEventValue, const bool bIsAutoZoomAnimated = true);
 
 	/** Updates the track's auto-zoom. Returns true if viewport was changed. Sets bIsAutoZoomDirty=true if needs another update. */
-	bool UpdateAutoZoomEx(const float InTopY, const float InBottomY, const double InMinEventValue, const double InMaxEventValue, const bool bIsAutoZoomAnimated);
+	UE_API bool UpdateAutoZoomEx(const float InTopY, const float InBottomY, const double InMinEventValue, const double InMaxEventValue, const bool bIsAutoZoomAnimated);
 
-	virtual FString FormatValue(double Value) const;
+	UE_API virtual FString FormatValue(double Value) const;
 
 private:
 	FText Name;
@@ -252,3 +254,5 @@ protected:
 };
 
 ////////////////////////////////////////////////////////////////////////////////////////////////////
+
+#undef UE_API
diff --git a/Engine/Source/Developer/TraceInsights/Public/Insights/ViewModels/GraphTrack.h b/Engine/Source/Developer/TraceInsights/Public/Insights/ViewModels/GraphTrack.h
--- a/Engine/Source/Developer/TraceInsights/Public/Insights/ViewModels/GraphTrack.h
+++ b/Engine/Source/Developer/TraceInsights/Public/Insights/ViewModels/GraphTrack.h
@@ -12,6 +12,8 @@
 #include "Insights/ViewModels/BaseTimingTrack.h"
 #include "Insights/ViewModels/GraphSeries.h"
 
+#define UE_API TRACEINSIGHTS_API
+
 class FMenuBuilder;
 struct FSlateBrush;
 
@@ -50,11 +52,11 @@ ENUM_CLASS_FLAGS(EGraphOptions);
 
 ////////////////////////////////////////////////////////////////////////////////////////////////////
 
-class TRACEINSIGHTS_API FGraphTrack : public FBaseTimingTrack
+class FGraphTrack : public FBaseTimingTrack
 {
 	friend class FGraphTrackBuilder;
 
-	INSIGHTS_DECLARE_RTTI(FGraphTrack, FBaseTimingTrack)
+	INSIGHTS_DECLARE_RTTI(FGraphTrack, FBaseTimingTrack, UE_API)
 
 private:
 	// Visual size of points (in pixels).
@@ -65,9 +67,9 @@ private:
 	static constexpr float PointSizeY = 3.0f;
 
 public:
-	explicit FGraphTrack();
-	explicit FGraphTrack(const FString& InName);
-	virtual ~FGraphTrack();
+	UE_API explicit FGraphTrack();
+	UE_API explicit FGraphTrack(const FString& InName);
+	UE_API virtual ~FGraphTrack();
 
 	//////////////////////////////////////////////////
 	// Options
@@ -90,19 +92,19 @@ public:
 	//////////////////////////////////////////////////
 	// FBaseTimingTrack
 
-	virtual void Reset() override;
+	UE_API virtual void Reset() override;
 
-	virtual void PostUpdate(const ITimingTrackUpdateContext& Context) override;
+	UE_API virtual void PostUpdate(const ITimingTrackUpdateContext& Context) override;
 
-	virtual void PreDraw(const ITimingTrackDrawContext& Context) const override;
-	virtual void Draw(const ITimingTrackDrawContext& Context) const override;
-	virtual void DrawEvent(const ITimingTrackDrawContext& Context, const ITimingEvent& InTimingEvent, EDrawEventMode InDrawMode) const override;
+	UE_API virtual void PreDraw(const ITimingTrackDrawContext& Context) const override;
+	UE_API virtual void Draw(const ITimingTrackDrawContext& Context) const override;
+	UE_API virtual void DrawEvent(const ITimingTrackDrawContext& Context, const ITimingEvent& InTimingEvent, EDrawEventMode InDrawMode) const override;
 
-	virtual void InitTooltip(FTooltipDrawState& InOutTooltip, const ITimingEvent& InTooltipEvent) const override;
+	UE_API virtual void InitTooltip(FTooltipDrawState& InOutTooltip, const ITimingEvent& InTooltipEvent) const override;
 
-	virtual const TSharedPtr<const ITimingEvent> GetEvent(float InPosX, float InPosY, const FTimingTrackViewport& Viewport) const override;
+	UE_API virtual const TSharedPtr<const ITimingEvent> GetEvent(float InPosX, float InPosY, const FTimingTrackViewport& Viewport) const override;
 
-	virtual void BuildContextMenu(FMenuBuilder& MenuBuilder) override;
+	UE_API virtual void BuildContextMenu(FMenuBuilder& MenuBuilder) override;
 
 	//////////////////////////////////////////////////
 
@@ -120,23 +122,23 @@ public:
 	int32 GetNumDrawBoxes() const { return NumDrawBoxes; }
 
 protected:
-	void UpdateStats();
+	UE_API void UpdateStats();
 
-	void DrawSeries(const FGraphSeries& Series, UE::Insights::FDrawContext& DrawContext, const FTimingTrackViewport& Viewport) const;
+	UE_API void DrawSeries(const FGraphSeries& Series, UE::Insights::FDrawContext& DrawContext, const FTimingTrackViewport& Viewport) const;
 
-	virtual void DrawVerticalAxisGrid(const ITimingTrackDrawContext& Context) const;
-	virtual void DrawHeader(const ITimingTrackDrawContext& Context) const;
+	UE_API virtual void DrawVerticalAxisGrid(const ITimingTrackDrawContext& Context) const;
+	UE_API virtual void DrawHeader(const ITimingTrackDrawContext& Context) const;
 
 	// Get the Y value that is used to provide a clipping border between adjacent graph tracks.
 	virtual float GetBorderY() const { return 0.0f; }
 
-	bool ContextMenu_ToggleOption_CanExecute(EGraphOptions Option);
-	virtual void ContextMenu_ToggleOption_Execute(EGraphOptions Option);
-	bool ContextMenu_ToggleOption_IsChecked(EGraphOptions Option);
+	UE_API bool ContextMenu_ToggleOption_CanExecute(EGraphOptions Option);
+	UE_API virtual void ContextMenu_ToggleOption_Execute(EGraphOptions Option);
+	UE_API bool ContextMenu_ToggleOption_IsChecked(EGraphOptions Option);
 
 private:
-	bool ContextMenu_ShowPointsWithBorder_CanExecute();
-	bool ContextMenu_UseEventDuration_CanExecute();
+	UE_API bool ContextMenu_ShowPointsWithBorder_CanExecute();
+	UE_API bool ContextMenu_UseEventDuration_CanExecute();
 
 protected:
 	TArray<TSharedPtr<FGraphSeries>> AllSeries;
@@ -166,20 +168,22 @@ protected:
 
 ////////////////////////////////////////////////////////////////////////////////////////////////////
 
-class TRACEINSIGHTS_API FRandomGraphTrack : public FGraphTrack
+class FRandomGraphTrack : public FGraphTrack
 {
-	INSIGHTS_DECLARE_RTTI(FRandomGraphTrack, FGraphTrack)
+	INSIGHTS_DECLARE_RTTI(FRandomGraphTrack, FGraphTrack, UE_API)
 
 public:
-	FRandomGraphTrack();
-	virtual ~FRandomGraphTrack();
+	UE_API FRandomGraphTrack();
+	UE_API virtual ~FRandomGraphTrack();
 
-	virtual void Update(const ITimingTrackUpdateContext& Context) override;
+	UE_API virtual void Update(const ITimingTrackUpdateContext& Context) override;
 
-	void AddDefaultSeries();
+	UE_API void AddDefaultSeries();
 
 protected:
-	void GenerateSeries(FGraphSeries& Series, const FTimingTrackViewport& Viewport, const int32 EventCount, int32 Seed);
+	UE_API void GenerateSeries(FGraphSeries& Series, const FTimingTrackViewport& Viewport, const int32 EventCount, int32 Seed);
 };
 
 ////////////////////////////////////////////////////////////////////////////////////////////////////
+
+#undef UE_API
diff --git a/Engine/Source/Developer/TraceInsights/Public/Insights/ViewModels/GraphTrackBuilder.h b/Engine/Source/Developer/TraceInsights/Public/Insights/ViewModels/GraphTrackBuilder.h
--- a/Engine/Source/Developer/TraceInsights/Public/Insights/ViewModels/GraphTrackBuilder.h
+++ b/Engine/Source/Developer/TraceInsights/Public/Insights/ViewModels/GraphTrackBuilder.h
@@ -4,13 +4,15 @@
 
 #include "CoreMinimal.h"
 
+#define UE_API TRACEINSIGHTS_API
+
 class FGraphTrack;
 class FGraphSeries;
 class FTimingTrackViewport;
 
 ////////////////////////////////////////////////////////////////////////////////////////////////////
 
-class TRACEINSIGHTS_API FGraphTrackBuilder
+class FGraphTrackBuilder
 {
 private:
 	struct FPointInfo
@@ -23,8 +25,8 @@ private:
 	};
 
 public:
-	explicit FGraphTrackBuilder(FGraphTrack& InTrack, FGraphSeries& InSeries, const FTimingTrackViewport& InViewport);
-	~FGraphTrackBuilder();
+	UE_API explicit FGraphTrackBuilder(FGraphTrack& InTrack, FGraphSeries& InSeries, const FTimingTrackViewport& InViewport);
+	UE_API ~FGraphTrackBuilder();
 
 	/**
 	 * Non-copyable
@@ -36,24 +38,24 @@ public:
 	FGraphSeries& GetSeries() const { return Series; }
 	const FTimingTrackViewport& GetViewport() const { return Viewport; }
 
-	void AddEvent(double Time, double Duration, double Value, bool bConnected = true);
+	UE_API void AddEvent(double Time, double Duration, double Value, bool bConnected = true);
 
 private:
-	void BeginPoints();
-	bool AddPoint(double Time, double Value);
-	void FlushPoints();
-	void EndPoints();
-
-	void BeginConnectedLines();
-	void AddConnectedLine(double Time, double Value, bool bNewBatch);
-	void FlushConnectedLine();
-	void AddConnectedLinePoint(float X, float Y);
-	void EndConnectedLines();
-
-	void BeginBoxes();
-	void AddBox(double Time, double Duration, double Value);
-	void FlushBox();
-	void EndBoxes();
+	UE_API void BeginPoints();
+	UE_API bool AddPoint(double Time, double Value);
+	UE_API void FlushPoints();
+	UE_API void EndPoints();
+
+	UE_API void BeginConnectedLines();
+	UE_API void AddConnectedLine(double Time, double Value, bool bNewBatch);
+	UE_API void FlushConnectedLine();
+	UE_API void AddConnectedLinePoint(float X, float Y);
+	UE_API void EndConnectedLines();
+
+	UE_API void BeginBoxes();
+	UE_API void AddBox(double Time, double Duration, double Value);
+	UE_API void FlushBox();
+	UE_API void EndBoxes();
 
 private:
 	FGraphTrack& Track;
@@ -77,3 +79,5 @@ private:
 };
 
 ////////////////////////////////////////////////////////////////////////////////////////////////////
+
+#undef UE_API
diff --git a/Engine/Source/Developer/TraceInsights/Public/Insights/ViewModels/GraphTrackEvent.h b/Engine/Source/Developer/TraceInsights/Public/Insights/ViewModels/GraphTrackEvent.h
--- a/Engine/Source/Developer/TraceInsights/Public/Insights/ViewModels/GraphTrackEvent.h
+++ b/Engine/Source/Developer/TraceInsights/Public/Insights/ViewModels/GraphTrackEvent.h
@@ -7,12 +7,14 @@
 // TraceInsights
 #include "Insights/ViewModels/ITimingEvent.h"
 
+#define UE_API TRACEINSIGHTS_API
+
 class FGraphTrack;
 class FGraphSeries;
 
 ////////////////////////////////////////////////////////////////////////////////////////////////////
 
-struct TRACEINSIGHTS_API FGraphSeriesEvent
+struct FGraphSeriesEvent
 {
 	double Time;
 	double Duration;
@@ -26,9 +28,9 @@ struct TRACEINSIGHTS_API FGraphSeriesEvent
 
 ////////////////////////////////////////////////////////////////////////////////////////////////////
 
-class TRACEINSIGHTS_API FGraphTrackEvent : public ITimingEvent
+class FGraphTrackEvent : public ITimingEvent
 {
-	INSIGHTS_DECLARE_RTTI(FGraphTrackEvent, ITimingEvent)
+	INSIGHTS_DECLARE_RTTI(FGraphTrackEvent, ITimingEvent, UE_API)
 
 public:
 	explicit FGraphTrackEvent(const TSharedRef<const FGraphTrack> InTrack, const TSharedRef<const FGraphSeries> InSeries, const FGraphSeriesEvent& InSeriesEvent)
@@ -40,7 +42,7 @@ public:
 	//////////////////////////////////////////////////
 	// ITimingEvent interface
 
-	virtual const TSharedRef<const FBaseTimingTrack> GetTrack() const override;
+	UE_API virtual const TSharedRef<const FBaseTimingTrack> GetTrack() const override;
 
 	virtual uint32 GetDepth() const override { return 0; }
 
@@ -77,3 +79,5 @@ private:
 };
 
 ////////////////////////////////////////////////////////////////////////////////////////////////////
+
+#undef UE_API
diff --git a/Engine/Source/Developer/TraceInsights/Public/Insights/ViewModels/ITimingEvent.h b/Engine/Source/Developer/TraceInsights/Public/Insights/ViewModels/ITimingEvent.h
--- a/Engine/Source/Developer/TraceInsights/Public/Insights/ViewModels/ITimingEvent.h
+++ b/Engine/Source/Developer/TraceInsights/Public/Insights/ViewModels/ITimingEvent.h
@@ -24,9 +24,9 @@ namespace UE::Insights { class FDrawContext; }
 
 ////////////////////////////////////////////////////////////////////////////////////////////////////
 
-class TRACEINSIGHTS_API ITimingEvent
+class ITimingEvent
 {
-	INSIGHTS_DECLARE_RTTI_BASE(ITimingEvent)
+	INSIGHTS_DECLARE_RTTI_BASE(ITimingEvent, TRACEINSIGHTS_API)
 
 public:
 	virtual const TSharedRef<const FBaseTimingTrack> GetTrack() const = 0;
@@ -54,9 +54,9 @@ public:
 
 ////////////////////////////////////////////////////////////////////////////////////////////////////
 
-class TRACEINSIGHTS_API ITimingEventFilter
+class ITimingEventFilter
 {
-	INSIGHTS_DECLARE_RTTI_BASE(ITimingEventFilter)
+	INSIGHTS_DECLARE_RTTI_BASE(ITimingEventFilter, TRACEINSIGHTS_API)
 
 public:
 	/**
@@ -78,9 +78,9 @@ public:
 
 ////////////////////////////////////////////////////////////////////////////////////////////////////
 
-class TRACEINSIGHTS_API ITimingEventRelation
+class ITimingEventRelation
 {
-	INSIGHTS_DECLARE_RTTI_BASE(ITimingEventRelation)
+	INSIGHTS_DECLARE_RTTI_BASE(ITimingEventRelation, TRACEINSIGHTS_API)
 
 public:
 	enum class EDrawFilter
diff --git a/Engine/Source/Developer/TraceInsights/Public/Insights/ViewModels/ThreadTrackEvent.h b/Engine/Source/Developer/TraceInsights/Public/Insights/ViewModels/ThreadTrackEvent.h
--- a/Engine/Source/Developer/TraceInsights/Public/Insights/ViewModels/ThreadTrackEvent.h
+++ b/Engine/Source/Developer/TraceInsights/Public/Insights/ViewModels/ThreadTrackEvent.h
@@ -7,24 +7,28 @@
 // TraceInsights
 #include "Insights/ViewModels/TimingEvent.h"
 
+#define UE_API TRACEINSIGHTS_API
+
 ////////////////////////////////////////////////////////////////////////////////////////////////////
 
-class TRACEINSIGHTS_API FThreadTrackEvent : public FTimingEvent
+class FThreadTrackEvent : public FTimingEvent
 {
-	INSIGHTS_DECLARE_RTTI(FThreadTrackEvent, FTimingEvent)
+	INSIGHTS_DECLARE_RTTI(FThreadTrackEvent, FTimingEvent, UE_API)
 
 public:
-	FThreadTrackEvent(const TSharedRef<const FBaseTimingTrack> InTrack, double InStartTime, double InEndTime, uint32 InDepth);
+	UE_API FThreadTrackEvent(const TSharedRef<const FBaseTimingTrack> InTrack, double InStartTime, double InEndTime, uint32 InDepth);
 	virtual ~FThreadTrackEvent() {}
 
-	uint32 GetTimerIndex() const;
-	void SetTimerIndex(uint32 InTimerIndex);
+	UE_API uint32 GetTimerIndex() const;
+	UE_API void SetTimerIndex(uint32 InTimerIndex);
 
-	uint32 GetTimerId() const;
-	void SetTimerId(uint32 InTimerId);
+	UE_API uint32 GetTimerId() const;
+	UE_API void SetTimerId(uint32 InTimerId);
 
 private:
 	uint32 TimerIndex = 0;
 };
 
 ////////////////////////////////////////////////////////////////////////////////////////////////////
+
+#undef UE_API
diff --git a/Engine/Source/Developer/TraceInsights/Public/Insights/ViewModels/TimingEvent.h b/Engine/Source/Developer/TraceInsights/Public/Insights/ViewModels/TimingEvent.h
--- a/Engine/Source/Developer/TraceInsights/Public/Insights/ViewModels/TimingEvent.h
+++ b/Engine/Source/Developer/TraceInsights/Public/Insights/ViewModels/TimingEvent.h
@@ -14,13 +14,15 @@
 #include "Insights/ViewModels/ITimingEvent.h"
 #include "Insights/ViewModels/TimingEventSearch.h"
 
+#define UE_API TRACEINSIGHTS_API
+
 class FBaseTimingTrack;
 
 ////////////////////////////////////////////////////////////////////////////////////////////////////
 
-class TRACEINSIGHTS_API FTimingEvent : public ITimingEvent
+class FTimingEvent : public ITimingEvent
 {
-	INSIGHTS_DECLARE_RTTI(FTimingEvent, ITimingEvent)
+	INSIGHTS_DECLARE_RTTI(FTimingEvent, ITimingEvent, UE_API)
 
 public:
 	FTimingEvent(const TSharedRef<const FBaseTimingTrack> InTrack, double InStartTime, double InEndTime, uint32 InDepth, uint64 InType = uint64(-1))
@@ -76,8 +78,8 @@ public:
 
 	uint64 GetType() const { return Type; }
 
-	static uint32 ComputeEventColor(uint32 Id);
-	static uint32 ComputeEventColor(const TCHAR* Str);
+	static UE_API uint32 ComputeEventColor(uint32 Id);
+	static UE_API uint32 ComputeEventColor(const TCHAR* Str);
 
 protected:
 	void SetType(uint64 InType) { Type = InType; }
@@ -112,7 +114,7 @@ private:
 
 class FAcceptNoneTimingEventFilter : public ITimingEventFilter
 {
-	INSIGHTS_DECLARE_RTTI(FAcceptNoneTimingEventFilter, ITimingEventFilter)
+	INSIGHTS_DECLARE_RTTI(FAcceptNoneTimingEventFilter, ITimingEventFilter, UE_API)
 
 public:
 	FAcceptNoneTimingEventFilter() {}
@@ -128,7 +130,7 @@ public:
 
 class FAcceptAllTimingEventFilter : public ITimingEventFilter
 {
-	INSIGHTS_DECLARE_RTTI(FAcceptAllTimingEventFilter, ITimingEventFilter)
+	INSIGHTS_DECLARE_RTTI(FAcceptAllTimingEventFilter, ITimingEventFilter, UE_API)
 
 public:
 	FAcceptAllTimingEventFilter() {}
@@ -144,7 +146,7 @@ public:
 
 class FAggregatedTimingEventFilter : public ITimingEventFilter
 {
-	INSIGHTS_DECLARE_RTTI(FAggregatedTimingEventFilter, ITimingEventFilter)
+	INSIGHTS_DECLARE_RTTI(FAggregatedTimingEventFilter, ITimingEventFilter, UE_API)
 
 public:
 	FAggregatedTimingEventFilter() : ChangeNumber(0) {}
@@ -172,7 +174,7 @@ protected:
 
 class FAllAggregatedTimingEventFilter : public FAggregatedTimingEventFilter
 {
-	INSIGHTS_DECLARE_RTTI(FAllAggregatedTimingEventFilter, FAggregatedTimingEventFilter)
+	INSIGHTS_DECLARE_RTTI(FAllAggregatedTimingEventFilter, FAggregatedTimingEventFilter, UE_API)
 
 public:
 	FAllAggregatedTimingEventFilter() {}
@@ -224,7 +226,7 @@ public:
 
 class FAnyAggregatedTimingEventFilter : public FAggregatedTimingEventFilter
 {
-	INSIGHTS_DECLARE_RTTI(FAnyAggregatedTimingEventFilter, FAggregatedTimingEventFilter)
+	INSIGHTS_DECLARE_RTTI(FAnyAggregatedTimingEventFilter, FAggregatedTimingEventFilter, UE_API)
 
 public:
 	FAnyAggregatedTimingEventFilter() {}
@@ -274,9 +276,9 @@ public:
 
 ////////////////////////////////////////////////////////////////////////////////////////////////////
 
-class TRACEINSIGHTS_API FTimingEventFilter : public ITimingEventFilter
+class FTimingEventFilter : public ITimingEventFilter
 {
-	INSIGHTS_DECLARE_RTTI(FTimingEventFilter, ITimingEventFilter)
+	INSIGHTS_DECLARE_RTTI(FTimingEventFilter, ITimingEventFilter, UE_API)
 
 public:
 	FTimingEventFilter()
@@ -298,7 +300,7 @@ public:
 	//////////////////////////////////////////////////
 	// ITimingEventFilter interface
 
-	virtual bool FilterTrack(const FBaseTimingTrack& InTrack) const override;
+	UE_API virtual bool FilterTrack(const FBaseTimingTrack& InTrack) const override;
 	virtual bool FilterEvent(const ITimingEvent& InEvent) const override { return true; }
 	virtual bool FilterEvent(double InEventStartTime, double InEventEndTime, uint32 InEventDepth, const TCHAR* InEventName, uint64 InEventType = 0, uint32 InEventColor = 0) const override { return true; }
 	virtual uint32 GetChangeNumber() const override { return ChangeNumber; }
@@ -331,7 +333,7 @@ protected:
 
 class FTimingEventFilterByMinDuration : public FTimingEventFilter
 {
-	INSIGHTS_DECLARE_RTTI(FTimingEventFilterByMinDuration, FTimingEventFilter)
+	INSIGHTS_DECLARE_RTTI(FTimingEventFilterByMinDuration, FTimingEventFilter, UE_API)
 
 public:
 	FTimingEventFilterByMinDuration(double InMinDuration) : MinDuration(InMinDuration) {}
@@ -363,7 +365,7 @@ private:
 
 class FTimingEventFilterByMaxDuration : public FTimingEventFilter
 {
-	INSIGHTS_DECLARE_RTTI(FTimingEventFilterByMaxDuration, FTimingEventFilter)
+	INSIGHTS_DECLARE_RTTI(FTimingEventFilterByMaxDuration, FTimingEventFilter, UE_API)
 
 public:
 	FTimingEventFilterByMaxDuration(double InMaxDuration) : MaxDuration(InMaxDuration) {}
@@ -393,9 +395,9 @@ private:
 
 ////////////////////////////////////////////////////////////////////////////////////////////////////
 
-class TRACEINSIGHTS_API FTimingEventFilterByEventType : public FTimingEventFilter
+class FTimingEventFilterByEventType : public FTimingEventFilter
 {
-	INSIGHTS_DECLARE_RTTI(FTimingEventFilterByEventType, FTimingEventFilter)
+	INSIGHTS_DECLARE_RTTI(FTimingEventFilterByEventType, FTimingEventFilter, UE_API)
 
 public:
 	FTimingEventFilterByEventType(uint64 InEventType) : EventType(InEventType) {}
@@ -432,7 +434,7 @@ private:
 
 class FTimingEventFilterByFrameIndex : public FTimingEventFilter
 {
-	INSIGHTS_DECLARE_RTTI(FTimingEventFilterByFrameIndex, FTimingEventFilter)
+	INSIGHTS_DECLARE_RTTI(FTimingEventFilterByFrameIndex, FTimingEventFilter, UE_API)
 
 public:
 	FTimingEventFilterByFrameIndex(uint64 InFrameIndex) : FrameIndex(InFrameIndex) {}
@@ -466,3 +468,5 @@ private:
 };
 
 ////////////////////////////////////////////////////////////////////////////////////////////////////
+
+#undef UE_API
diff --git a/Engine/Source/Developer/TraceInsights/Public/Insights/ViewModels/TimingEventsTrack.h b/Engine/Source/Developer/TraceInsights/Public/Insights/ViewModels/TimingEventsTrack.h
--- a/Engine/Source/Developer/TraceInsights/Public/Insights/ViewModels/TimingEventsTrack.h
+++ b/Engine/Source/Developer/TraceInsights/Public/Insights/ViewModels/TimingEventsTrack.h
@@ -10,11 +10,13 @@
 // TraceInsights
 #include "Insights/ViewModels/BaseTimingTrack.h"
 
+#define UE_API TRACEINSIGHTS_API
+
 ////////////////////////////////////////////////////////////////////////////////////////////////////
 
 struct FSlateBrush;
 
-class TRACEINSIGHTS_API ITimingEventsTrackDrawStateBuilder
+class ITimingEventsTrackDrawStateBuilder
 {
 public:
 	typedef TFunctionRef<const FString(float /*AvailableWidth*/)> GetEventNameCallback;
@@ -28,30 +30,30 @@ public:
 
 ////////////////////////////////////////////////////////////////////////////////////////////////////
 
-class TRACEINSIGHTS_API FTimingEventsTrack : public FBaseTimingTrack
+class FTimingEventsTrack : public FBaseTimingTrack
 {
-	INSIGHTS_DECLARE_RTTI(FTimingEventsTrack, FBaseTimingTrack)
+	INSIGHTS_DECLARE_RTTI(FTimingEventsTrack, FBaseTimingTrack, UE_API)
 
 public:
-	explicit FTimingEventsTrack();
-	explicit FTimingEventsTrack(const FString& InTrackName);
-	virtual ~FTimingEventsTrack();
+	UE_API explicit FTimingEventsTrack();
+	UE_API explicit FTimingEventsTrack(const FString& InTrackName);
+	UE_API virtual ~FTimingEventsTrack();
 
 	//////////////////////////////////////////////////
 	// FBaseTimingTrack
 
-	virtual void Reset() override;
+	UE_API virtual void Reset() override;
 
-	virtual void PreUpdate(const ITimingTrackUpdateContext& Context) override;
-	virtual void PostUpdate(const ITimingTrackUpdateContext& Context) override;
+	UE_API virtual void PreUpdate(const ITimingTrackUpdateContext& Context) override;
+	UE_API virtual void PostUpdate(const ITimingTrackUpdateContext& Context) override;
 
-	virtual void Draw(const ITimingTrackDrawContext& Context) const override;
-	virtual void PostDraw(const ITimingTrackDrawContext& Context) const override;
-	virtual void DrawEvent(const ITimingTrackDrawContext& Context, const ITimingEvent& InTimingEvent, EDrawEventMode InDrawMode) const override;
+	UE_API virtual void Draw(const ITimingTrackDrawContext& Context) const override;
+	UE_API virtual void PostDraw(const ITimingTrackDrawContext& Context) const override;
+	UE_API virtual void DrawEvent(const ITimingTrackDrawContext& Context, const ITimingEvent& InTimingEvent, EDrawEventMode InDrawMode) const override;
 
-	virtual const TSharedPtr<const ITimingEvent> GetEvent(float InPosX, float InPosY, const FTimingTrackViewport& Viewport) const override;
+	UE_API virtual const TSharedPtr<const ITimingEvent> GetEvent(float InPosX, float InPosY, const FTimingTrackViewport& Viewport) const override;
 
-	virtual TSharedPtr<ITimingEventFilter> GetFilterByEvent(const TSharedPtr<const ITimingEvent> InTimingEvent) const override;
+	UE_API virtual TSharedPtr<ITimingEventFilter> GetFilterByEvent(const TSharedPtr<const ITimingEvent> InTimingEvent) const override;
 
 	//////////////////////////////////////////////////
 
@@ -79,20 +81,20 @@ protected:
 		}
 	}
 
-	void UpdateTrackHeight(const ITimingTrackUpdateContext& Context);
+	UE_API void UpdateTrackHeight(const ITimingTrackUpdateContext& Context);
 
-	void DrawEvents(const ITimingTrackDrawContext& Context, const float OffsetY = 1.0f) const;
-	void DrawHeader(const ITimingTrackDrawContext& Context) const;
+	UE_API void DrawEvents(const ITimingTrackDrawContext& Context, const float OffsetY = 1.0f) const;
+	UE_API void DrawHeader(const ITimingTrackDrawContext& Context) const;
 
-	void DrawMarkers(const ITimingTrackDrawContext& Context, float LineY, float LineH) const;
+	UE_API void DrawMarkers(const ITimingTrackDrawContext& Context, float LineY, float LineH) const;
 
-	void DrawSelectedEventInfo(const FString& InText, const FTimingTrackViewport& Viewport, const UE::Insights::FDrawContext& DrawContext, const FSlateBrush* WhiteBrush, const FSlateFontInfo& Font) const;
-	void DrawSelectedEventInfoEx(const FString& InText, const FString& InLeftText, const FString& InTopText, const FTimingTrackViewport& Viewport, const UE::Insights::FDrawContext& DrawContext, const FSlateBrush* WhiteBrush, const FSlateFontInfo& Font) const;
+	UE_API void DrawSelectedEventInfo(const FString& InText, const FTimingTrackViewport& Viewport, const UE::Insights::FDrawContext& DrawContext, const FSlateBrush* WhiteBrush, const FSlateFontInfo& Font) const;
+	UE_API void DrawSelectedEventInfoEx(const FString& InText, const FString& InLeftText, const FString& InTopText, const FTimingTrackViewport& Viewport, const UE::Insights::FDrawContext& DrawContext, const FSlateBrush* WhiteBrush, const FSlateFontInfo& Font) const;
 
-	int32 GetHeaderBackgroundLayerId(const ITimingTrackDrawContext& Context) const;
-	int32 GetHeaderTextLayerId(const ITimingTrackDrawContext& Context) const;
+	UE_API int32 GetHeaderBackgroundLayerId(const ITimingTrackDrawContext& Context) const;
+	UE_API int32 GetHeaderTextLayerId(const ITimingTrackDrawContext& Context) const;
 
-	virtual const TSharedPtr<const ITimingEvent> GetEvent(double InTime, double SecondsPerPixel, int32 Depth) const;
+	UE_API virtual const TSharedPtr<const ITimingEvent> GetEvent(double InTime, double SecondsPerPixel, int32 Depth) const;
 
 	virtual bool HasCustomFilter() const { return false; }
 
@@ -117,7 +119,9 @@ private:
 	FFilteredDrawStateInfo FilteredDrawStateInfo;
 
 public:
-	static bool bUseDownSampling; // toggle to enable/disable downsampling, for debugging purposes only
+	static UE_API bool bUseDownSampling; // toggle to enable/disable downsampling, for debugging purposes only
 };
 
 ////////////////////////////////////////////////////////////////////////////////////////////////////
+
+#undef UE_API
diff --git a/Engine/Source/Developer/TraceInsights/Public/Insights/ViewModels/TimingTrackViewport.h b/Engine/Source/Developer/TraceInsights/Public/Insights/ViewModels/TimingTrackViewport.h
--- a/Engine/Source/Developer/TraceInsights/Public/Insights/ViewModels/TimingTrackViewport.h
+++ b/Engine/Source/Developer/TraceInsights/Public/Insights/ViewModels/TimingTrackViewport.h
@@ -7,6 +7,8 @@
 // TraceInsights
 #include "Insights/ViewModels/TimingViewLayout.h"
 
+#define UE_API TRACEINSIGHTS_API
+
 class SScrollBar;
 
 ////////////////////////////////////////////////////////////////////////////////////////////////////
@@ -58,7 +60,7 @@ ENUM_CLASS_FLAGS(ETimingTrackViewportDirtyFlags);
 
 ////////////////////////////////////////////////////////////////////////////////////////////////////
 
-class TRACEINSIGHTS_API FTimingTrackViewport
+class FTimingTrackViewport
 {
 public:
 	FTimingTrackViewport()
@@ -97,7 +99,7 @@ public:
 	/** Height of viewport, in pixels. [px] */
 	float GetHeight() const { return Height; }
 
-	bool UpdateSize(const float InWidth, const float InHeight);
+	UE_API bool UpdateSize(const float InWidth, const float InHeight);
 
 	/** Minimum session time, in seconds. [s] */
 	double GetMinValidTime() const { return MinValidTime; }
@@ -160,25 +162,25 @@ public:
 		return AlignTimeToPixel(Time, ScaleX);
 	}
 
-	double RestrictEndTime(const double InEndTime) const;
-	double RestrictDuration(const double InStartTime, const double InEndTime) const;
+	UE_API double RestrictEndTime(const double InEndTime) const;
+	UE_API double RestrictDuration(const double InStartTime, const double InEndTime) const;
 
 	//////////////////////////////////////////////////
 	// Scrolling
 
-	bool ScrollAtTime(const double Time);
-	bool CenterOnTimeInterval(const double Time, const double Duration);
+	UE_API bool ScrollAtTime(const double Time);
+	UE_API bool CenterOnTimeInterval(const double Time, const double Duration);
 
-	void GetHorizontalScrollLimits(double& OutMinT, double& OutMaxT);
-	bool EnforceHorizontalScrollLimits(const double U);
+	UE_API void GetHorizontalScrollLimits(double& OutMinT, double& OutMaxT);
+	UE_API bool EnforceHorizontalScrollLimits(const double U);
 
 	//////////////////////////////////////////////////
 	// Zooming
 
-	bool ZoomOnTimeInterval(const double Time, const double Duration);
-	bool RelativeZoomWithFixedX(const float Delta, const float X);
-	bool ZoomWithFixedX(const double NewScaleX, const float X);
-	bool SetScaleX(const double NewScaleX);
+	UE_API bool ZoomOnTimeInterval(const double Time, const double Duration);
+	UE_API bool RelativeZoomWithFixedX(const float Delta, const float X);
+	UE_API bool ZoomWithFixedX(const double NewScaleX, const float X);
+	UE_API bool SetScaleX(const double NewScaleX);
 
 	//////////////////////////////////////////////////
 	// Vertical axis
@@ -208,7 +210,7 @@ public:
 	float GetViewportY(const float Y) const { return PosY + TopOffset + Y - ScrollPosY; }
 
 	const FTimingViewLayout& GetLayout() const { return Layout; }
-	void UpdateLayout();
+	UE_API void UpdateLayout();
 
 	bool IsLayoutCompactModeEnabled() const { return Layout.bIsCompactMode; }
 
@@ -235,11 +237,11 @@ public:
 
 	//////////////////////////////////////////////////
 
-	bool OnUserScrolled(TSharedPtr<SScrollBar> ScrollBar, float ScrollOffset);
-	void UpdateScrollBar(TSharedPtr<SScrollBar> ScrollBar) const;
+	UE_API bool OnUserScrolled(TSharedPtr<SScrollBar> ScrollBar, float ScrollOffset);
+	UE_API void UpdateScrollBar(TSharedPtr<SScrollBar> ScrollBar) const;
 
-	bool OnUserScrolledY(TSharedPtr<SScrollBar> ScrollBar, float ScrollOffset);
-	void UpdateScrollBarY(TSharedPtr<SScrollBar> ScrollBar) const;
+	UE_API bool OnUserScrolledY(TSharedPtr<SScrollBar> ScrollBar, float ScrollOffset);
+	UE_API void UpdateScrollBarY(TSharedPtr<SScrollBar> ScrollBar) const;
 
 	//////////////////////////////////////////////////
 
@@ -276,3 +278,5 @@ private:
 };
 
 ////////////////////////////////////////////////////////////////////////////////////////////////////
+
+#undef UE_API
diff --git a/Engine/Source/Developer/TraceInsights/Public/Insights/ViewModels/TimingViewLayout.h b/Engine/Source/Developer/TraceInsights/Public/Insights/ViewModels/TimingViewLayout.h
--- a/Engine/Source/Developer/TraceInsights/Public/Insights/ViewModels/TimingViewLayout.h
+++ b/Engine/Source/Developer/TraceInsights/Public/Insights/ViewModels/TimingViewLayout.h
@@ -5,9 +5,11 @@
 #include "CoreMinimal.h"
 #include "HAL/Platform.h"
 
+#define UE_API TRACEINSIGHTS_API
+
 ////////////////////////////////////////////////////////////////////////////////////////////////////
 
-struct TRACEINSIGHTS_API FTimingViewLayout
+struct FTimingViewLayout
 {
 	static constexpr float RealMinTimelineH = 7.0f;
 
@@ -90,9 +92,11 @@ struct TRACEINSIGHTS_API FTimingViewLayout
 		}
 	}
 
-	void ForceNormalMode();
-	void ForceCompactMode();
-	bool Update();
+	UE_API void ForceNormalMode();
+	UE_API void ForceCompactMode();
+	UE_API bool Update();
 };
 
 ////////////////////////////////////////////////////////////////////////////////////////////////////
+
+#undef UE_API
diff --git a/Engine/Source/Developer/TraceInsights/Public/Insights/ViewModels/TooltipDrawState.h b/Engine/Source/Developer/TraceInsights/Public/Insights/ViewModels/TooltipDrawState.h
--- a/Engine/Source/Developer/TraceInsights/Public/Insights/ViewModels/TooltipDrawState.h
+++ b/Engine/Source/Developer/TraceInsights/Public/Insights/ViewModels/TooltipDrawState.h
@@ -9,6 +9,8 @@
 #include "Math/Color.h"
 #include "Templates/SharedPointer.h"
 
+#define UE_API TRACEINSIGHTS_API
+
 ////////////////////////////////////////////////////////////////////////////////////////////////////
 
 struct FSlateBrush;
@@ -17,7 +19,7 @@ namespace UE::Insights { class FDrawContext; }
 
 ////////////////////////////////////////////////////////////////////////////////////////////////////
 
-class TRACEINSIGHTS_API FTooltipDrawState
+class FTooltipDrawState
 {
 public:
 	static constexpr float BorderX = 6.0f;
@@ -29,9 +31,9 @@ public:
 	static constexpr float DefaultLineHeight = 14.0f;
 	static constexpr float NameValueDX = 2.0f;
 
-	static FLinearColor DefaultTitleColor;
-	static FLinearColor DefaultNameColor;
-	static FLinearColor DefaultValueColor;
+	static UE_API FLinearColor DefaultTitleColor;
+	static UE_API FLinearColor DefaultNameColor;
+	static UE_API FLinearColor DefaultValueColor;
 
 private:
 	enum class FDrawTextType
@@ -53,18 +55,18 @@ private:
 	};
 
 public:
-	FTooltipDrawState();
-	~FTooltipDrawState();
+	UE_API FTooltipDrawState();
+	UE_API ~FTooltipDrawState();
 
-	void Reset();
+	UE_API void Reset();
 
-	void ResetContent();
-	void AddTitle(FStringView Title);
-	void AddTitle(FStringView Title, const FLinearColor& Color);
-	void AddNameValueTextLine(FStringView Name, FStringView Value);
-	void AddTextLine(FStringView Text, const FLinearColor& Color);
-	void AddTextLine(const float X, const float Y, FStringView Text, const FLinearColor& Color);
-	void UpdateLayout(); // updates ValueOffsetX and DesiredSize
+	UE_API void ResetContent();
+	UE_API void AddTitle(FStringView Title);
+	UE_API void AddTitle(FStringView Title, const FLinearColor& Color);
+	UE_API void AddNameValueTextLine(FStringView Name, FStringView Value);
+	UE_API void AddTextLine(FStringView Text, const FLinearColor& Color);
+	UE_API void AddTextLine(const float X, const float Y, FStringView Text, const FLinearColor& Color);
+	UE_API void UpdateLayout(); // updates ValueOffsetX and DesiredSize
 
 	const FLinearColor& GetBackgroundColor() const { return BackgroundColor; }
 	void SetBackgroundColor(const FLinearColor& InBackgroundColor) { BackgroundColor = InBackgroundColor; }
@@ -78,13 +80,13 @@ public:
 	const FVector2D& GetSize() const { return Size; }
 	const FVector2D& GetDesiredSize() const { return DesiredSize; }
 
-	void Update(); // updates current Opacity and Size
+	UE_API void Update(); // updates current Opacity and Size
 
 	const FVector2D& GetPosition() const { return Position; }
-	void SetPosition(const FVector2D& MousePosition, const float MinX, const float MaxX, const float MinY, const float MaxY);
+	UE_API void SetPosition(const FVector2D& MousePosition, const float MinX, const float MaxX, const float MinY, const float MaxY);
 	void SetPosition(const float PosX, const float PosY) { Position.X = PosX; Position.Y = PosY; }
 
-	void Draw(const UE::Insights::FDrawContext& DrawContext) const;
+	UE_API void Draw(const UE::Insights::FDrawContext& DrawContext) const;
 
 	void SetFontScale(float InFontScale) { FontScale = InFontScale; }
 	float GetFontScale() const { return FontScale; }
@@ -115,3 +117,5 @@ private:
 };
 
 ////////////////////////////////////////////////////////////////////////////////////////////////////
+
+#undef UE_API
diff --git a/Engine/Source/Developer/TraceInsights/Public/Insights/Widgets/SMajorTabWindow.h b/Engine/Source/Developer/TraceInsights/Public/Insights/Widgets/SMajorTabWindow.h
--- a/Engine/Source/Developer/TraceInsights/Public/Insights/Widgets/SMajorTabWindow.h
+++ b/Engine/Source/Developer/TraceInsights/Public/Insights/Widgets/SMajorTabWindow.h
@@ -12,6 +12,8 @@
 // TraceInsights
 #include "Insights/IUnrealInsightsModule.h" // for FInsightsMajorTabExtender
 
+#define UE_API TRACEINSIGHTS_API
+
 ////////////////////////////////////////////////////////////////////////////////////////////////////
 
 class FActiveTimerHandle;
@@ -26,14 +28,14 @@ namespace Insights
 {
 
 /** Implements the base class for a major tab window. */
-class TRACEINSIGHTS_API SMajorTabWindow : public SCompoundWidget
+class SMajorTabWindow : public SCompoundWidget
 {
 public:
 	/** Default constructor. */
-	SMajorTabWindow(const FName& InMajorTabId);
+	UE_API SMajorTabWindow(const FName& InMajorTabId);
 
 	/** Virtual destructor. */
-	virtual ~SMajorTabWindow();
+	UE_API virtual ~SMajorTabWindow();
 
 	SLATE_BEGIN_ARGS(SMajorTabWindow) {}
 	SLATE_END_ARGS()
@@ -43,15 +45,15 @@ public:
 	TSharedPtr<FWorkspaceItem> GetWorkspaceMenuGroup() const { return WorkspaceMenuGroup; }
 	const TSharedPtr<FUICommandList> GetCommandList() const { return CommandList; }
 
-	virtual void Reset();
+	UE_API virtual void Reset();
 
 	/** Constructs this widget. */
-	void Construct(const FArguments& InArgs, const TSharedRef<SDockTab>& ConstructUnderMajorTab, const TSharedPtr<SWindow>& ConstructUnderWindow);
+	UE_API void Construct(const FArguments& InArgs, const TSharedRef<SDockTab>& ConstructUnderMajorTab, const TSharedPtr<SWindow>& ConstructUnderWindow);
 
-	void ShowTab(const FName& TabId);
-	void HideTab(const FName& TabId);
+	UE_API void ShowTab(const FName& TabId);
+	UE_API void HideTab(const FName& TabId);
 	void ShowHideTab(const FName& TabId, bool bShow) { bShow ? ShowTab(TabId) : HideTab(TabId); }
-	void CloseAllOpenTabs();
+	UE_API void CloseAllOpenTabs();
 
 	////////////////////////////////////////////////////////////////////////////////////////////////////
 	// OnWindowClosedEvent
@@ -64,16 +66,16 @@ private:
 	FWindowClosedEvent WindowClosedEvent;
 
 protected:
-	virtual const TCHAR* GetAnalyticsEventName() const;
-	virtual TSharedRef<FWorkspaceItem> CreateWorkspaceMenuGroup();
-	virtual void RegisterTabSpawners();
-	virtual TSharedRef<FTabManager::FLayout> CreateDefaultTabLayout() const;
-	virtual TSharedRef<SWidget> CreateToolbar(TSharedPtr<FExtender> Extender);
+	UE_API virtual const TCHAR* GetAnalyticsEventName() const;
+	UE_API virtual TSharedRef<FWorkspaceItem> CreateWorkspaceMenuGroup();
+	UE_API virtual void RegisterTabSpawners();
+	UE_API virtual TSharedRef<FTabManager::FLayout> CreateDefaultTabLayout() const;
+	UE_API virtual TSharedRef<SWidget> CreateToolbar(TSharedPtr<FExtender> Extender);
 
 	void SetCommandList(const TSharedPtr<FUICommandList> InCommandList) { CommandList = InCommandList; }
 
-	void AddOpenTab(const TSharedRef<SDockTab>& DockTab);
-	void RemoveOpenTab(const TSharedRef<SDockTab>& DockTab);
+	UE_API void AddOpenTab(const TSharedRef<SDockTab>& DockTab);
+	UE_API void RemoveOpenTab(const TSharedRef<SDockTab>& DockTab);
 
 private:
 	/**
@@ -82,13 +84,13 @@ private:
 	 * @param MenuBuilder The multi-box builder that should be filled with content for this pull-down menu.
 	 * @param TabManager A Tab Manager from which to populate tab spawner menu items.
 	 */
-	static void FillMenu(FMenuBuilder& MenuBuilder, const TSharedPtr<FTabManager> TabManager);
+	static UE_API void FillMenu(FMenuBuilder& MenuBuilder, const TSharedPtr<FTabManager> TabManager);
 
 	/** Returns true if the current analysis session is valid. */
-	bool IsValidSession() const;
+	UE_API bool IsValidSession() const;
 
 	/** Updates the amount of time the window has been active. */
-	EActiveTimerReturnType UpdateActiveDuration(double InCurrentTime, float InDeltaTime);
+	UE_API EActiveTimerReturnType UpdateActiveDuration(double InCurrentTime, float InDeltaTime);
 
 	/**
 	 * The system will use this event to notify a widget that the cursor has entered it. This event is NOT bubbled.
@@ -96,14 +98,14 @@ private:
 	 * @param MyGeometry The Geometry of the widget receiving the event
 	 * @param MouseEvent Information about the input event
 	 */
-	virtual void OnMouseEnter(const FGeometry& MyGeometry, const FPointerEvent& MouseEvent) override;
+	UE_API virtual void OnMouseEnter(const FGeometry& MyGeometry, const FPointerEvent& MouseEvent) override;
 
 	/**
 	 * The system will use this event to notify a widget that the cursor has left it. This event is NOT bubbled.
 	 *
 	 * @param MouseEvent Information about the input event
 	 */
-	virtual void OnMouseLeave(const FPointerEvent& MouseEvent) override;
+	UE_API virtual void OnMouseLeave(const FPointerEvent& MouseEvent) override;
 
 	/**
 	 * Called after a key is pressed when this widget has focus
@@ -113,7 +115,7 @@ private:
 	 *
 	 * @return  Returns whether the event was handled, along with other possible actions
 	 */
-	virtual FReply OnKeyDown(const FGeometry& MyGeometry, const FKeyEvent& InKeyEvent) override;
+	UE_API virtual FReply OnKeyDown(const FGeometry& MyGeometry, const FKeyEvent& InKeyEvent) override;
 
 	/**
 	 * Called when the user is dropping something onto a widget; terminates drag and drop.
@@ -123,7 +125,7 @@ private:
 	 *
 	 * @return A reply that indicated whether this event was handled.
 	 */
-	virtual FReply OnDrop(const FGeometry& MyGeometry, const FDragDropEvent& DragDropEvent) override;
+	UE_API virtual FReply OnDrop(const FGeometry& MyGeometry, const FDragDropEvent& DragDropEvent) override;
 
 	/**
 	 * Called during drag and drop when the the mouse is being dragged over a widget.
@@ -133,7 +135,7 @@ private:
 	 *
 	 * @return A reply that indicated whether this event was handled.
 	 */
-	virtual FReply OnDragOver(const FGeometry& MyGeometry, const FDragDropEvent& DragDropEvent)  override;
+	UE_API virtual FReply OnDragOver(const FGeometry& MyGeometry, const FDragDropEvent& DragDropEvent)  override;
 
 private:
 	/** The major tab id for this window. */
@@ -162,3 +164,5 @@ private:
 };
 
 } // namespace Insights
+
+#undef UE_API
diff --git a/Engine/Source/Developer/TraceInsightsCore/Public/InsightsCore/Common/AsyncOperationProgress.h b/Engine/Source/Developer/TraceInsightsCore/Public/InsightsCore/Common/AsyncOperationProgress.h
--- a/Engine/Source/Developer/TraceInsightsCore/Public/InsightsCore/Common/AsyncOperationProgress.h
+++ b/Engine/Source/Developer/TraceInsightsCore/Public/InsightsCore/Common/AsyncOperationProgress.h
@@ -9,7 +9,7 @@
 namespace UE::Insights
 {
 
-class TRACEINSIGHTSCORE_API IAsyncOperationProgress
+class IAsyncOperationProgress
 {
 public:
 	virtual void CancelCurrentAsyncOp() = 0;
@@ -17,7 +17,7 @@ public:
 	virtual bool ShouldCancelAsyncOp() const = 0;
 };
 
-class TRACEINSIGHTSCORE_API FAsyncOperationProgress : public IAsyncOperationProgress
+class FAsyncOperationProgress : public IAsyncOperationProgress
 {
 public:
 	virtual void CancelCurrentAsyncOp() override { bCancelCurrentAsyncOp.store(true); }
diff --git a/Engine/Source/Developer/TraceInsightsCore/Public/InsightsCore/Common/AvailabilityCheck.h b/Engine/Source/Developer/TraceInsightsCore/Public/InsightsCore/Common/AvailabilityCheck.h
--- a/Engine/Source/Developer/TraceInsightsCore/Public/InsightsCore/Common/AvailabilityCheck.h
+++ b/Engine/Source/Developer/TraceInsightsCore/Public/InsightsCore/Common/AvailabilityCheck.h
@@ -4,23 +4,25 @@
 
 #include "CoreTypes.h"
 
+#define UE_API TRACEINSIGHTSCORE_API
+
 namespace UE::Insights
 {
 
 /**
  * Utility class used by profiler managers to limit how often they check for availability conditions.
  */
-class TRACEINSIGHTSCORE_API FAvailabilityCheck
+class FAvailabilityCheck
 {
 public:
 	/** Returns true if managers are allowed to do (slow) availability check during this tick. */
-	bool Tick();
+	UE_API bool Tick();
 
 	/** Disables the "availability check" (i.e. Tick() calls will return false when disabled). */
-	void Disable();
+	UE_API void Disable();
 
 	/** Enables the "availability check" with a specified initial delay. */
-	void Enable(double InWaitTime);
+	UE_API void Enable(double InWaitTime);
 
 private:
 	double WaitTime = 0.0;
@@ -28,3 +30,5 @@ private:
 };
 
 } // namespace UE::Insights
+
+#undef UE_API
diff --git a/Engine/Source/Developer/TraceInsightsCore/Public/InsightsCore/Common/IAsyncOperationStatusProvider.h b/Engine/Source/Developer/TraceInsightsCore/Public/InsightsCore/Common/IAsyncOperationStatusProvider.h
--- a/Engine/Source/Developer/TraceInsightsCore/Public/InsightsCore/Common/IAsyncOperationStatusProvider.h
+++ b/Engine/Source/Developer/TraceInsightsCore/Public/InsightsCore/Common/IAsyncOperationStatusProvider.h
@@ -9,7 +9,7 @@ class FText;
 namespace UE::Insights
 {
 
-class TRACEINSIGHTSCORE_API IAsyncOperationStatusProvider
+class IAsyncOperationStatusProvider
 {
 public:
 	virtual bool IsRunning() const = 0;
diff --git a/Engine/Source/Developer/TraceInsightsCore/Public/InsightsCore/Common/InsightsCoreStyle.h b/Engine/Source/Developer/TraceInsightsCore/Public/InsightsCore/Common/InsightsCoreStyle.h
--- a/Engine/Source/Developer/TraceInsightsCore/Public/InsightsCore/Common/InsightsCoreStyle.h
+++ b/Engine/Source/Developer/TraceInsightsCore/Public/InsightsCore/Common/InsightsCoreStyle.h
@@ -11,6 +11,8 @@
 #include "Types/SlateEnums.h"
 #include "UObject/NameTypes.h"
 
+#define UE_API TRACEINSIGHTSCORE_API
+
 struct FSlateBrush;
 
 namespace UE::Insights
@@ -51,14 +53,14 @@ public:
 		}
 	}
 
-	class TRACEINSIGHTSCORE_API FStyle : public FSlateStyleSet
+	class FStyle : public FSlateStyleSet
 	{
 	public:
-		FStyle(const FName& InStyleSetName);
+		UE_API FStyle(const FName& InStyleSetName);
 
-		void Initialize();
+		UE_API void Initialize();
 
-		void SyncParentStyles();
+		UE_API void SyncParentStyles();
 
 	public:
 		// Styles inherited from the parent style
@@ -77,3 +79,5 @@ private:
 };
 
 } // namespace UE::Insights
+
+#undef UE_API
diff --git a/Engine/Source/Developer/TraceInsightsCore/Public/InsightsCore/Common/MessageDialogUtils.h b/Engine/Source/Developer/TraceInsightsCore/Public/InsightsCore/Common/MessageDialogUtils.h
--- a/Engine/Source/Developer/TraceInsightsCore/Public/InsightsCore/Common/MessageDialogUtils.h
+++ b/Engine/Source/Developer/TraceInsightsCore/Public/InsightsCore/Common/MessageDialogUtils.h
@@ -4,6 +4,8 @@
 
 #include "Internationalization/Text.h"
 
+#define UE_API TRACEINSIGHTSCORE_API
+
 namespace UE::Insights
 {
 
@@ -17,12 +19,14 @@ enum class EDialogResponse
 
 ////////////////////////////////////////////////////////////////////////////////////////////////////
 
-class TRACEINSIGHTSCORE_API FMessageDialogUtils
+class FMessageDialogUtils
 {
 public:
-	static EDialogResponse ShowChoiceDialog(const FText& Title, const FText& Content);
+	static UE_API EDialogResponse ShowChoiceDialog(const FText& Title, const FText& Content);
 };
 
 ////////////////////////////////////////////////////////////////////////////////////////////////////
 
-} // namespace UE::Insights
\ No newline at end of file
+} // namespace UE::Insights
+
+#undef UE_API
diff --git a/Engine/Source/Developer/TraceInsightsCore/Public/InsightsCore/Common/MiscUtils.h b/Engine/Source/Developer/TraceInsightsCore/Public/InsightsCore/Common/MiscUtils.h
--- a/Engine/Source/Developer/TraceInsightsCore/Public/InsightsCore/Common/MiscUtils.h
+++ b/Engine/Source/Developer/TraceInsightsCore/Public/InsightsCore/Common/MiscUtils.h
@@ -4,10 +4,12 @@
 
 #include "CoreTypes.h"
 
+#define UE_API TRACEINSIGHTSCORE_API
+
 namespace UE::Insights
 {
 
-class TRACEINSIGHTSCORE_API FMiscUtils
+class FMiscUtils
 {
 public:
 	/**
@@ -15,7 +17,9 @@ public:
 	 * @param InCmdLine - The command line passed to the new UnrealInsights.exe process
 	 * @returns true if UnrealInsights process is created successfully.
 	 */
-	static bool OpenUnrealInsights(const TCHAR* InCmdLine = nullptr);
+	static UE_API bool OpenUnrealInsights(const TCHAR* InCmdLine = nullptr);
 };
 
 } // namespace UE::Insights
+
+#undef UE_API
diff --git a/Engine/Source/Developer/TraceInsightsCore/Public/InsightsCore/Filter/ViewModels/FilterConfigurator.h b/Engine/Source/Developer/TraceInsightsCore/Public/InsightsCore/Filter/ViewModels/FilterConfigurator.h
--- a/Engine/Source/Developer/TraceInsightsCore/Public/InsightsCore/Filter/ViewModels/FilterConfigurator.h
+++ b/Engine/Source/Developer/TraceInsightsCore/Public/InsightsCore/Filter/ViewModels/FilterConfigurator.h
@@ -8,32 +8,34 @@
 #include "InsightsCore/Filter/ViewModels/Filters.h"
 #include "InsightsCore/Filter/ViewModels/IFilterExecutor.h"
 
+#define UE_API TRACEINSIGHTSCORE_API
+
 namespace UE::Insights
 {
 
 ////////////////////////////////////////////////////////////////////////////////////////////////////
 
-class TRACEINSIGHTSCORE_API FFilterConfigurator : public IFilterExecutor
+class FFilterConfigurator : public IFilterExecutor
 {
 public:
-	FFilterConfigurator();
+	UE_API FFilterConfigurator();
 
-	FFilterConfigurator(const FFilterConfigurator& Other);
-	FFilterConfigurator& operator=(const FFilterConfigurator& Other);
+	UE_API FFilterConfigurator(const FFilterConfigurator& Other);
+	UE_API FFilterConfigurator& operator=(const FFilterConfigurator& Other);
 
-	bool operator==(const FFilterConfigurator& Other) const;
+	UE_API bool operator==(const FFilterConfigurator& Other) const;
 	bool operator!=(const FFilterConfigurator& Other) const { return !(*this == Other); }
 
-	virtual ~FFilterConfigurator();
+	UE_API virtual ~FFilterConfigurator();
 
 	bool IsEmpty() const { return RootNode->GetChildrenCount() == 0; }
 
 	/** Called to update the internal state of some filters */
 	void Update() { RootNode->Update(); }
 
-	virtual bool ApplyFilters(const FFilterContext& Context) const override;
+	UE_API virtual bool ApplyFilters(const FFilterContext& Context) const override;
 
-	bool IsKeyUsed(int32 Key) const;
+	UE_API bool IsKeyUsed(int32 Key) const;
 
 	FFilterConfiguratorNodePtr GetRootNode() { return RootNode; }
 	TSharedPtr<TArray<TSharedPtr<FFilter>>>& GetAvailableFilters() { return AvailableFilters; }
@@ -65,7 +67,7 @@ private:
 	////////////////////////////////////////////////////////////////////////////////////////////////////
 
 private:
-	void ComputeUsedKeys();
+	UE_API void ComputeUsedKeys();
 
 private:
 	FFilterConfiguratorNodePtr RootNode;
@@ -78,3 +80,5 @@ private:
 ////////////////////////////////////////////////////////////////////////////////////////////////////
 
 } // namespace UE::Insights
+
+#undef UE_API
diff --git a/Engine/Source/Developer/TraceInsightsCore/Public/InsightsCore/Filter/ViewModels/FilterConfiguratorNode.h b/Engine/Source/Developer/TraceInsightsCore/Public/InsightsCore/Filter/ViewModels/FilterConfiguratorNode.h
--- a/Engine/Source/Developer/TraceInsightsCore/Public/InsightsCore/Filter/ViewModels/FilterConfiguratorNode.h
+++ b/Engine/Source/Developer/TraceInsightsCore/Public/InsightsCore/Filter/ViewModels/FilterConfiguratorNode.h
@@ -7,6 +7,8 @@
 #include "InsightsCore/Filter/ViewModels/Filters.h"
 #include "InsightsCore/Table/ViewModels/BaseTreeNode.h"
 
+#define UE_API TRACEINSIGHTSCORE_API
+
 namespace UE::Insights
 {
 
@@ -28,50 +30,50 @@ typedef TWeakPtr<class FFilterConfiguratorNode> FFilterConfiguratorNodeWeak;
 /**
  * Class used to store information about a filter node
  */
-class TRACEINSIGHTSCORE_API FFilterConfiguratorNode : public FBaseTreeNode
+class FFilterConfiguratorNode : public FBaseTreeNode
 {
-	INSIGHTS_DECLARE_RTTI(FFilterConfiguratorNode, FBaseTreeNode)
+	INSIGHTS_DECLARE_RTTI(FFilterConfiguratorNode, FBaseTreeNode, UE_API)
 
 public:
 	/** Initialization constructor for the filter configurator node. */
-	FFilterConfiguratorNode(const FName InName, bool bInIsGroup);
+	UE_API FFilterConfiguratorNode(const FName InName, bool bInIsGroup);
 
-	static TSharedRef<FFilterConfiguratorNode> DeepCopy(const FFilterConfiguratorNode& Node);
+	static UE_API TSharedRef<FFilterConfiguratorNode> DeepCopy(const FFilterConfiguratorNode& Node);
 
-	bool operator==(const FFilterConfiguratorNode& Other) const;
+	UE_API bool operator==(const FFilterConfiguratorNode& Other) const;
 
 	virtual ~FFilterConfiguratorNode() {}
 
-	void SetAvailableFilters(TSharedPtr<TArray<TSharedPtr<FFilter>>> InAvailableFilters);
+	UE_API void SetAvailableFilters(TSharedPtr<TArray<TSharedPtr<FFilter>>> InAvailableFilters);
 	TSharedPtr<TArray<TSharedPtr<FFilter>>> GetAvailableFilters() { return AvailableFilters; }
 
-	void SetSelectedFilter(TSharedPtr<FFilter> InSelectedFilter);
+	UE_API void SetSelectedFilter(TSharedPtr<FFilter> InSelectedFilter);
 	TSharedPtr<FFilter> GetSelectedFilter() const { return SelectedFilter; }
 
-	void SetSelectedFilterOperator(TSharedPtr<IFilterOperator> InSelectedFilterOperator);
-	TSharedPtr<const IFilterOperator> GetSelectedFilterOperator() const;
+	UE_API void SetSelectedFilterOperator(TSharedPtr<IFilterOperator> InSelectedFilterOperator);
+	UE_API TSharedPtr<const IFilterOperator> GetSelectedFilterOperator() const;
 
-	const TArray<TSharedPtr<FFilterGroupOperator>>& GetFilterGroupOperators() const;
+	UE_API const TArray<TSharedPtr<FFilterGroupOperator>>& GetFilterGroupOperators() const;
 	void SetSelectedFilterGroupOperator(TSharedPtr<FFilterGroupOperator> InSelectedFilterGroupOperator) { SelectedFilterGroupOperator = InSelectedFilterGroupOperator; }
 	TSharedPtr<FFilterGroupOperator> GetSelectedFilterGroupOperator() const { return SelectedFilterGroupOperator; }
 
 	TSharedPtr<TArray<TSharedPtr<IFilterOperator>>> GetAvailableFilterOperators() const { return AvailableFilterOperators; }
 
 	const FString& GetTextBoxValue() { return TextBoxValue; }
-	void SetTextBoxValue(const FString& InValue);
+	UE_API void SetTextBoxValue(const FString& InValue);
 
-	bool ApplyFilters(const class FFilterContext& Context) const;
+	UE_API bool ApplyFilters(const class FFilterContext& Context) const;
 
-	void GetUsedKeys(TSet<int32>& GetUsedKeys) const;
+	UE_API void GetUsedKeys(TSet<int32>& GetUsedKeys) const;
 
-	void ProcessFilter();
+	UE_API void ProcessFilter();
 
-	void Update();
+	UE_API void Update();
 
 	TSharedPtr<FFilterState> GetSelectedFilterState() { return FilterState; }
 
 private:
-	FFilterConfiguratorNode& operator=(const FFilterConfiguratorNode& Other);
+	UE_API FFilterConfiguratorNode& operator=(const FFilterConfiguratorNode& Other);
 
 	TSharedPtr<TArray<TSharedPtr<FFilter>>> AvailableFilters;
 
@@ -89,3 +91,5 @@ private:
 ////////////////////////////////////////////////////////////////////////////////////////////////////
 
 } // namespace UE::Insights
+
+#undef UE_API
