Improved the performance of FLODUtilities::StripLODGeometry by:
- Replacing a custom hash key with standard engine type, FInt32Vector3
- Reducing memory allocations in a cache structure from several thousand to almost zero

diff --git a/Engine/Source/Developer/SkeletalMeshUtilitiesCommon/Private/LODUtilities.cpp b/Engine/Source/Developer/SkeletalMeshUtilitiesCommon/Private/LODUtilities.cpp
--- a/Engine/Source/Developer/SkeletalMeshUtilitiesCommon/Private/LODUtilities.cpp
+++ b/Engine/Source/Developer/SkeletalMeshUtilitiesCommon/Private/LODUtilities.cpp
@@ -4530,29 +4530,7 @@ bool FLODUtilities::StripLODGeometry(USkeletalMesh* SkeletalMesh, const int32 LO
 
 		const TArray< uint32 >& SoftVertexIndexToImportDataPointIndex = LODModel.GetRawPointIndices();
 
-		TMap<uint64, TArray<int32>> OptimizedFaceFinder;
-		
-		auto GetMatchFaceIndex = [&OptimizedFaceFinder, &ImportedData](const int32 FaceVertexA, const int32 FaceVertexB, int32 FaceVertexC)->int32
-		{
-			uint64 Key = (uint64)FaceVertexA | ((uint64)FaceVertexB >> 32) | (((uint64)FaceVertexC & 0xFFFF) >> 48);
-			TArray<int32>& FaceIndices = OptimizedFaceFinder.FindChecked(Key);
-			for (int32 PossibleFaceIndex = 0; PossibleFaceIndex < FaceIndices.Num(); ++PossibleFaceIndex)
-			{
-				int32 FaceIndex = FaceIndices[PossibleFaceIndex];
-				const SkeletalMeshImportData::FTriangle& Face = ImportedData.Faces[FaceIndex];
-				if (FaceVertexA == ImportedData.Wedges[Face.WedgeIndex[0]].VertexIndex)
-				{
-					if (FaceVertexB == ImportedData.Wedges[Face.WedgeIndex[1]].VertexIndex)
-					{
-						if (FaceVertexC == ImportedData.Wedges[Face.WedgeIndex[2]].VertexIndex)
-						{
-							return FaceIndex;
-						}
-					}
-				}
-			}
-			return INDEX_NONE;
-		};
+		TMap<FInt32Vector3, int32> OptimizedFaceFinder;
 
 		for (int32 FaceIndex = 0; FaceIndex < ImportedData.Faces.Num(); ++FaceIndex)
 		{
@@ -4560,9 +4538,8 @@ bool FLODUtilities::StripLODGeometry(USkeletalMesh* SkeletalMesh, const int32 LO
 			int32 FaceVertexA = ImportedData.Wedges[Face.WedgeIndex[0]].VertexIndex;
 			int32 FaceVertexB = ImportedData.Wedges[Face.WedgeIndex[1]].VertexIndex;
 			int32 FaceVertexC = ImportedData.Wedges[Face.WedgeIndex[2]].VertexIndex;
-			uint64 Key = (uint64)FaceVertexA | ((uint64)FaceVertexB >> 32) | (((uint64)FaceVertexC & 0xFFFF) >> 48);
-			TArray<int32>& FaceIndices = OptimizedFaceFinder.FindOrAdd(Key);
-			FaceIndices.Add(FaceIndex);
+
+			OptimizedFaceFinder.Add(FInt32Vector3(FaceVertexA, FaceVertexB, FaceVertexC), FaceIndex);
 		}
 
 		int32 RemovedFaceCount = 0;
@@ -4596,14 +4573,11 @@ bool FLODUtilities::StripLODGeometry(USkeletalMesh* SkeletalMesh, const int32 LO
 					int32 ImportedPointIndexA = SoftVertexIndexToImportDataPointIndex[VertexIndexA];
 					int32 ImportedPointIndexB = SoftVertexIndexToImportDataPointIndex[VertexIndexB];
 					int32 ImportedPointIndexC = SoftVertexIndexToImportDataPointIndex[VertexIndexC];
-					int32 FaceIndex = GetMatchFaceIndex(ImportedPointIndexA, ImportedPointIndexB, ImportedPointIndexC);
-					if (FaceIndex != INDEX_NONE)
+					int32 FaceIndex = OptimizedFaceFinder.FindChecked(FInt32Vector3(ImportedPointIndexA, ImportedPointIndexB, ImportedPointIndexC));
+					if (!FaceToRemove[FaceIndex])
 					{
-						if (!FaceToRemove[FaceIndex])
-						{
-							FaceToRemove[FaceIndex] = true;
-							RemovedFaceCount++;
-						}
+						FaceToRemove[FaceIndex] = true;
+						RemovedFaceCount++;
 					}
 				}
 				else
@@ -4688,3 +4662,4 @@ bool FLODUtilities::StripLODGeometry(USkeletalMesh* SkeletalMesh, const int32 LO
 #undef LOCTEXT_NAMESPACE // "LODUtilities"
 
 #endif //WITH_EDITOR
+
