Add null check for AssetImportData in SimplifySkeletalMeshLOD

diff --git a/Engine/Source/Developer/SkeletalMeshUtilitiesCommon/Private/LODUtilities.cpp b/Engine/Source/Developer/SkeletalMeshUtilitiesCommon/Private/LODUtilities.cpp
--- a/Engine/Source/Developer/SkeletalMeshUtilitiesCommon/Private/LODUtilities.cpp
+++ b/Engine/Source/Developer/SkeletalMeshUtilitiesCommon/Private/LODUtilities.cpp
@@ -1748,7 +1748,8 @@ void FLODUtilities::SimplifySkeletalMeshLOD(
 
 			if (DesiredLOD == 0)
 			{
-				SkeletalMesh->GetLODInfo(DesiredLOD)->SourceImportFilename = SkeletalMesh->GetAssetImportData()->GetFirstFilename();
+				const UAssetImportData* AssetImportData = SkeletalMesh->GetAssetImportData();
+				SkeletalMesh->GetLODInfo(DesiredLOD)->SourceImportFilename = AssetImportData ? AssetImportData->GetFirstFilename() : FString();
 			}
 		}
 	}
