Cloth/Interchange - Fixed UnbindClothingAndBackup crash when re-importing a Skeletal Mesh with more than one cloth LOD.

[RN] Cloth/Interchange - Fixed UnbindClothingAndBackup crash when re-importing a Skeletal Mesh with more than one cloth LOD.

diff --git a/Engine/Source/Developer/SkeletalMeshUtilitiesCommon/Private/LODUtilities.cpp b/Engine/Source/Developer/SkeletalMeshUtilitiesCommon/Private/LODUtilities.cpp
--- a/Engine/Source/Developer/SkeletalMeshUtilitiesCommon/Private/LODUtilities.cpp
+++ b/Engine/Source/Developer/SkeletalMeshUtilitiesCommon/Private/LODUtilities.cpp
@@ -3482,29 +3482,38 @@ void FLODUtilities::BuildMorphTargets(USkeletalMesh* SkeletalMesh, const FMeshDe
 
 void FLODUtilities::UnbindClothingAndBackup(USkeletalMesh* SkeletalMesh, TArray<ClothingAssetUtils::FClothingAssetMeshBinding>& ClothingBindings)
 {
-	for (int32 LODIndex = 0; LODIndex < SkeletalMesh->GetImportedModel()->LODModels.Num(); ++LODIndex)
+	ClothingBindings.Reset();
+
+	if (ClothingAssetUtils::HasAnyMeshClothingAssetBindings(SkeletalMesh))  // Don't run the PostEditChange unless strictly necessary
 	{
-		TArray<ClothingAssetUtils::FClothingAssetMeshBinding> LODBindings;
-		UnbindClothingAndBackup(SkeletalMesh, LODBindings, LODIndex);
-		ClothingBindings.Append(LODBindings);
+		FScopedSkeletalMeshPostEditChange SkeletalMeshPostEditChange(SkeletalMesh);
+
+		for (int32 LODIndex = 0; LODIndex < SkeletalMesh->GetImportedModel()->LODModels.Num(); ++LODIndex)
+		{
+			TArray<ClothingAssetUtils::FClothingAssetMeshBinding> LODBindings;
+			UnbindClothingAndBackup(SkeletalMesh, LODBindings, LODIndex);
+			ClothingBindings.Append(LODBindings);
+		}
 	}
 }
 
 void FLODUtilities::UnbindClothingAndBackup(USkeletalMesh* SkeletalMesh, TArray<ClothingAssetUtils::FClothingAssetMeshBinding>& ClothingBindings, const int32 LODIndex)
 {
-	if (!SkeletalMesh->GetImportedModel()->LODModels.IsValidIndex(LODIndex))
+	if (!ClothingAssetUtils::HasAnyLodMeshClothingAssetBindings(SkeletalMesh, LODIndex))  // Early exit to avoid the FScopedSkeletalMeshPostEditChange whenever possible
 	{
 		return;
 	}
-	FSkeletalMeshLODModel& LODModel = SkeletalMesh->GetImportedModel()->LODModels[LODIndex];
-	//Store the clothBinding
+	// Guard the unbinding with a scoped PostEditChange
+	FScopedSkeletalMeshPostEditChange SkeletalMeshPostEditChange(SkeletalMesh);
+	// Store the clothBinding
 	ClothingAssetUtils::GetAllLodMeshClothingAssetBindings(SkeletalMesh, ClothingBindings, LODIndex);
-	//Unbind the Cloth for this LOD before we reduce it, we will put back the cloth after the reduction, if it still match the sections
+	// Unbind the Cloth for this LOD before we reduce it, we will put back the cloth after the reduction, if it still match the sections
+	FSkeletalMeshLODModel& LODModel = SkeletalMesh->GetImportedModel()->LODModels[LODIndex];
 	for (ClothingAssetUtils::FClothingAssetMeshBinding& Binding : ClothingBindings)
 	{
 		if (Binding.LODIndex == LODIndex)
 		{
-			//Use the UserSectionsData original section index, this will ensure we remap correctly the cloth if the reduction has change the number of sections
+			// Use the UserSectionsData original section index, this will ensure we remap correctly the cloth if the reduction has change the number of sections
 			int32 OriginalDataSectionIndex = LODModel.Sections[Binding.SectionIndex].OriginalDataSectionIndex;
 			if (Binding.Asset)
 			{
@@ -3512,16 +3521,27 @@ void FLODUtilities::UnbindClothingAndBackup(USkeletalMesh* SkeletalMesh, TArray<
 				Binding.SectionIndex = OriginalDataSectionIndex;
 			}
 			
-			FSkelMeshSourceSectionUserData& SectionUserData = LODModel.UserSectionsData.FindChecked(OriginalDataSectionIndex);
-			SectionUserData.ClothingData.AssetGuid = FGuid();
-			SectionUserData.ClothingData.AssetLodIndex = INDEX_NONE;
-			SectionUserData.CorrespondClothAssetIndex = INDEX_NONE;
+			if (FSkelMeshSourceSectionUserData* const SectionUserData = LODModel.UserSectionsData.Find(OriginalDataSectionIndex);
+				ensureMsgf(SectionUserData, TEXT("A FScopedSkeletalMeshPostEditChange guard must be used to ensure that the model isn't getting rebuilt while this function is called.")))
+			{
+				SectionUserData->ClothingData.AssetGuid = FGuid();
+				SectionUserData->ClothingData.AssetLodIndex = INDEX_NONE;
+				SectionUserData->CorrespondClothAssetIndex = INDEX_NONE;
+			}
 		}
 	}
 }
 
 void FLODUtilities::RestoreClothingFromBackup(USkeletalMesh* SkeletalMesh, TArray<ClothingAssetUtils::FClothingAssetMeshBinding>& ClothingBindings)
 {
+	if (!SkeletalMesh ||
+		!SkeletalMesh->GetImportedModel() ||
+		!ClothingBindings.Num())  // Early exit to avoid the FScopedSkeletalMeshPostEditChange whenever possible
+	{
+		return;
+	}
+	// Guard the binding with a scoped PostEditChange
+	FScopedSkeletalMeshPostEditChange SkeletalMeshPostEditChange(SkeletalMesh);
 	for (int32 LODIndex = 0; LODIndex < SkeletalMesh->GetImportedModel()->LODModels.Num(); ++LODIndex)
 	{
 		RestoreClothingFromBackup(SkeletalMesh, ClothingBindings, LODIndex);
@@ -3530,10 +3550,16 @@ void FLODUtilities::RestoreClothingFromBackup(USkeletalMesh* SkeletalMesh, TArra
 
 void FLODUtilities::RestoreClothingFromBackup(USkeletalMesh* SkeletalMesh, TArray<ClothingAssetUtils::FClothingAssetMeshBinding>& ClothingBindings, const int32 LODIndex)
 {
-	if (!SkeletalMesh->GetImportedModel()->LODModels.IsValidIndex(LODIndex))
+	if (!SkeletalMesh ||
+		!SkeletalMesh->GetImportedModel() ||
+		!SkeletalMesh->GetImportedModel()->LODModels.IsValidIndex(LODIndex) ||
+		!ClothingBindings.Num())  // Early exit to avoid the FScopedSkeletalMeshPostEditChange whenever possible
 	{
 		return;
 	}
+	// Guard the binding with a scoped PostEditChange
+	FScopedSkeletalMeshPostEditChange SkeletalMeshPostEditChange(SkeletalMesh);
+
 	FSkeletalMeshLODModel& LODModel = SkeletalMesh->GetImportedModel()->LODModels[LODIndex];
 	for (ClothingAssetUtils::FClothingAssetMeshBinding& Binding : ClothingBindings)
 	{
@@ -3547,10 +3573,13 @@ void FLODUtilities::RestoreClothingFromBackup(USkeletalMesh* SkeletalMesh, TArra
 			{
 				if (Binding.Asset->BindToSkeletalMesh(SkeletalMesh, Binding.LODIndex, SectionIndex, Binding.AssetInternalLodIndex))
 				{
-					//If successfull set back the section user data
-					FSkelMeshSourceSectionUserData& SectionUserData = LODModel.UserSectionsData.FindChecked(Binding.SectionIndex);
-					SectionUserData.CorrespondClothAssetIndex = LODModel.Sections[SectionIndex].CorrespondClothAssetIndex;
-					SectionUserData.ClothingData = LODModel.Sections[SectionIndex].ClothingData;
+					// If successfull set back the section user data
+					if (FSkelMeshSourceSectionUserData* const SectionUserData = LODModel.UserSectionsData.Find(Binding.SectionIndex);
+						ensureMsgf(SectionUserData, TEXT("A FScopedSkeletalMeshPostEditChange guard must be used to ensure that the model isn't getting rebuilt while this function is called.")))
+					{
+						SectionUserData->CorrespondClothAssetIndex = LODModel.Sections[SectionIndex].CorrespondClothAssetIndex;
+						SectionUserData->ClothingData = LODModel.Sections[SectionIndex].ClothingData;
+					}
 				}
 			}
 			break;
diff --git a/Engine/Source/Runtime/ClothingSystemRuntimeCommon/Private/ClothingAsset.cpp b/Engine/Source/Runtime/ClothingSystemRuntimeCommon/Private/ClothingAsset.cpp
--- a/Engine/Source/Runtime/ClothingSystemRuntimeCommon/Private/ClothingAsset.cpp
+++ b/Engine/Source/Runtime/ClothingSystemRuntimeCommon/Private/ClothingAsset.cpp
@@ -43,26 +43,27 @@
 #if WITH_EDITOR
 void ClothingAssetUtils::UnbindFromSkeletalMesh(const UClothingAssetBase& Asset, USkeletalMesh& SkeletalMesh)
 {
-	FScopedSkeletalMeshPostEditChange SkeletalMeshPostEditChange(&SkeletalMesh);
-
-	if (const FSkeletalMeshModel* const Mesh = SkeletalMesh.GetImportedModel())
+	if (HasAnyMeshClothingAssetBindings(&SkeletalMesh))
 	{
-		const int32 NumLods = Mesh->LODModels.Num();
-		for (int32 LodIndex = 0; LodIndex < NumLods; ++LodIndex)
+		FScopedSkeletalMeshPostEditChange SkeletalMeshPostEditChange(&SkeletalMesh);
+
+		if (const FSkeletalMeshModel* const Mesh = SkeletalMesh.GetImportedModel())
 		{
-			UnbindFromSkeletalMesh(Asset, SkeletalMesh, LodIndex);
+			const int32 NumLods = Mesh->LODModels.Num();
+			for (int32 LodIndex = 0; LodIndex < NumLods; ++LodIndex)
+			{
+				UnbindFromSkeletalMesh(Asset, SkeletalMesh, LodIndex);
+			}
 		}
 	}
 }
 
 void ClothingAssetUtils::UnbindFromSkeletalMesh(const UClothingAssetBase& Asset, USkeletalMesh& SkeletalMesh, const int32 MeshLodIndex)
 {
-	// Find the chunk(s) we created
-	if (FSkeletalMeshModel* const Mesh = SkeletalMesh.GetImportedModel())
+	if (HasAnyLodMeshClothingAssetBindings(&SkeletalMesh, MeshLodIndex))
 	{
-		bool bChangedMesh = false;
-
-		if (!Mesh->LODModels.IsValidIndex(MeshLodIndex))
+		FSkeletalMeshModel* const Mesh = SkeletalMesh.GetImportedModel();
+		if (!Mesh || !Mesh->LODModels.IsValidIndex(MeshLodIndex))
 		{
 			UE_LOG(LogClothingAsset, Warning, TEXT("Failed to remove clothing asset [%s] from mesh LOD%d as that LOD doesn't exist."),
 				*Asset.GetName(),
@@ -70,6 +71,7 @@ void ClothingAssetUtils::UnbindFromSkeletalMesh(const UClothingAssetBase& Asset,
 			return;
 		}
 
+		bool bChangedMesh = false;
 		FSkeletalMeshLODModel& LodModel = Mesh->LODModels[MeshLodIndex];
 
 		for (int32 SectionIdx = LodModel.Sections.Num() - 1; SectionIdx >= 0; --SectionIdx)
@@ -126,6 +128,33 @@ void ClothingAssetUtils::UnbindFromSkeletalMesh(const UClothingAssetBase& Asset,
 	}
 }
 
+bool ClothingAssetUtils::HasAnyMeshClothingAssetBindings(const USkeletalMesh* SkeletalMesh)
+{
+	if (const FSkeletalMeshModel* const MeshModel = SkeletalMesh ? SkeletalMesh->GetImportedModel() : nullptr)
+	{
+		for (int32 LodIndex = 0; LodIndex < MeshModel->LODModels.Num(); ++LodIndex)
+		{
+			if (MeshModel->LODModels[LodIndex].HasClothData())
+			{
+				return true;
+			}
+		}
+	}
+	return false;
+}
+
+bool ClothingAssetUtils::HasAnyLodMeshClothingAssetBindings(const USkeletalMesh* SkeletalMesh, int32 LodIndex)
+{
+	if (const FSkeletalMeshModel* const MeshModel = SkeletalMesh ? SkeletalMesh->GetImportedModel() : nullptr)
+	{
+		if (MeshModel->LODModels.IsValidIndex(LodIndex) && MeshModel->LODModels[LodIndex].HasClothData())
+		{
+			return true;
+		}
+	}
+	return false;
+}
+
 void ClothingAssetUtils::GetAllMeshClothingAssetBindings(const USkeletalMesh* SkeletalMesh, TArray<FClothingAssetMeshBinding>& OutBindings)
 {
 	OutBindings.Empty();
diff --git a/Engine/Source/Runtime/ClothingSystemRuntimeCommon/Public/ClothingAsset.h b/Engine/Source/Runtime/ClothingSystemRuntimeCommon/Public/ClothingAsset.h
--- a/Engine/Source/Runtime/ClothingSystemRuntimeCommon/Public/ClothingAsset.h
+++ b/Engine/Source/Runtime/ClothingSystemRuntimeCommon/Public/ClothingAsset.h
@@ -46,6 +46,21 @@ namespace ClothingAssetUtils
 	 */
 	UE_API void UnbindFromSkeletalMesh(const UClothingAssetBase& Asset, USkeletalMesh& SkeletalMesh, const int32 MeshLodIndex);
 
+	/**
+	 * Return whether there are any cloth bindings whithin the skeletal mesh model.
+	 * Uses the cloth mapping data as the source predicate to match the behavior of GetAllMeshClothingAssetBindings().
+	 * @param SkeletalMesh The skeletal mesh to check the cloth binding from.
+	 */
+	UE_API bool HasAnyMeshClothingAssetBindings(const USkeletalMesh* SkeletalMesh);
+
+	/**
+	 * Return whether there are any cloth bindings whithin the skeletal mesh model LOD.
+	 * Uses the cloth mapping data as the source predicate to match the behavior of GetAllLodMeshClothingAssetBindings().
+	 * @param SkeletalMesh The skeletal mesh to check the cloth binding from.
+	 * @param LodIndex The specified LOD.
+	 */
+	UE_API bool HasAnyLodMeshClothingAssetBindings(const USkeletalMesh* SkeletalMesh, const int32 LodIndex);
+
 	/**
 	 * Given a skeletal mesh model, find all of the currently bound clothing assets and their binding information.
 	 * @param SkeletalMesh The skeletal mesh to extract binding from.
