Skeletal Mesh: Search for orphaned morph targets using StaticFindObjectFastInternal to ensure we also find objects with the async flag set.

#rnx

diff --git a/Engine/Source/Developer/SkeletalMeshUtilitiesCommon/Private/LODUtilities.cpp b/Engine/Source/Developer/SkeletalMeshUtilitiesCommon/Private/LODUtilities.cpp
--- a/Engine/Source/Developer/SkeletalMeshUtilitiesCommon/Private/LODUtilities.cpp
+++ b/Engine/Source/Developer/SkeletalMeshUtilitiesCommon/Private/LODUtilities.cpp
@@ -3337,38 +3337,26 @@ static void BuildMorphTargetsInternal(
 			{
 				if (LODIndex == 0)
 				{
-					//Garbage collect must be delayed until the skeletal mesh build is done by registering to delegate
-					//FCoreUObjectDelegates::GetPreGarbageCollectDelegate.
-					if (ensure(!GIsGarbageCollecting))
-					{
-						if (!IsInGameThread())
-						{
-							//TODO remove this code when overriding a UObject will be allow outside of the game thread
-							//We currently need to avoid overriding an existing asset outside of the game thread
-							UObject* ExistingMorphTarget = StaticFindObject(UMorphTarget::StaticClass(), BaseSkelMesh, *ShapeName);
-							if (ExistingMorphTarget)
-							{
-								//make sure the object is not standalone or transactional
-								ExistingMorphTarget->ClearFlags(RF_Standalone | RF_Transactional);
-								//Move this object in the transient package
-								ExistingMorphTarget->Rename(nullptr, GetTransientPackage(), REN_DoNotDirty | REN_DontCreateRedirectors | REN_NonTransactional);
-								ExistingMorphTarget = nullptr;
-							}
-						}
+					FGCScopeGuard Guard;
 
-						MorphTarget = NewObject<UMorphTarget>(BaseSkelMesh, ObjectName);
+					// Check if there's an orphaned morph target object floating about. We use StaticFindObjectFastInternal
+					// rather than StaticFindObjectFast since we might have a situation where the orphaned morph target might
+					// still have the Async flag set, which StaticFindObjectFast will skip over.
+					MorphTarget = Cast<UMorphTarget>(StaticFindObjectFastInternal(UMorphTarget::StaticClass(), BaseSkelMesh, ObjectName, EFindObjectFlags::ExactClass));
+					if (MorphTarget)
+					{
+						// Temporarily set the async flag on it so that it doesn't vanish once we leave the GC scope.
+						MorphTarget->SetInternalFlags(EInternalObjectFlags::Async);
+						
+						// Eviscerate this morph target for 
+						MorphTarget->GetMorphLODModels().Reset();
 					}
 					else
 					{
-						UE_ASSET_LOG(LogLODUtilities, Error, BaseSkelMesh, TEXT("FLODUtilities::BuildMorphTargets: Garbage collection is running during the skeletal build. Morph target [%s] cannot be built properly and will be missing."), *ShapeName);
+						// This will already have the async flag on it, since we're creating the object in a worker thread.
+						MorphTarget = NewObject<UMorphTarget>(BaseSkelMesh, ObjectName);
 					}
 				}
-				else
-				{
-					/*AddTokenizedErrorMessage(FTokenizedMessage::Create(EMessageSeverity::Error, FText::Format(FText::FromString("Could not find the {0} morphtarget for LOD {1}. \
-						Make sure the name for morphtarget matches with LOD 0"), FText::FromString(ShapeName), FText::FromString(FString::FromInt(LODIndex)))),
-						FFbxErrors::SkeletalMesh_LOD_MissingMorphTarget);*/
-				}
 			}
 		}
 
