SkeletalMesh from mesh descriptions: ability to define the attributes to be enable for render at convert time
- also adds a FInitializationParams structure for extra flexibility

#rnx
#Synced-CL 49388435

diff --git a/Engine/Plugins/ChaosClothAsset/Source/ChaosClothAssetTools/Private/ChaosClothAsset/SkeletalMeshConverter.cpp b/Engine/Plugins/ChaosClothAsset/Source/ChaosClothAssetTools/Private/ChaosClothAsset/SkeletalMeshConverter.cpp
--- a/Engine/Plugins/ChaosClothAsset/Source/ChaosClothAssetTools/Private/ChaosClothAsset/SkeletalMeshConverter.cpp
+++ b/Engine/Plugins/ChaosClothAsset/Source/ChaosClothAssetTools/Private/ChaosClothAsset/SkeletalMeshConverter.cpp
@@ -126,13 +126,14 @@ bool UClothAssetEditorSkeletalMeshConverter::ExportToSkeletalMesh(const UChaosCl
 		SkeletalMaterials.Emplace(UMaterial::GetDefaultMaterial(MD_Surface));
 	}
 
-	constexpr bool bRecomputeNormals = false;
-	constexpr bool bRecomputeTangents = false;
+	FStaticToSkeletalMeshConverter::FInitializationParams Parameters;
+		Parameters.Materials = SkeletalMaterials;
+		Parameters.bRecomputeNormals = false;
+		Parameters.bRecomputeTangents = false;
+	
 	return FStaticToSkeletalMeshConverter::InitializeSkeletalMeshFromMeshDescriptions(
 		&SkeletalMesh,
 		MeshDescriptions,
-		SkeletalMaterials,
 		ClothAssetBase.GetRefSkeleton(),
-		bRecomputeNormals,
-		bRecomputeTangents);
+		Parameters);
 }
diff --git a/Engine/Plugins/Experimental/Dataflow/Source/DataflowNodes/Private/Dataflow/DataflowCollectionSkeletalMeshUtils.h b/Engine/Plugins/Experimental/Dataflow/Source/DataflowNodes/Private/Dataflow/DataflowCollectionSkeletalMeshUtils.h
--- a/Engine/Plugins/Experimental/Dataflow/Source/DataflowNodes/Private/Dataflow/DataflowCollectionSkeletalMeshUtils.h
+++ b/Engine/Plugins/Experimental/Dataflow/Source/DataflowNodes/Private/Dataflow/DataflowCollectionSkeletalMeshUtils.h
@@ -61,12 +61,14 @@ static bool CreateSkeletalMesh(USkeletalMesh* SkeletalMesh, const TArray<UE::Geo
 		MaterialView = Materials;
 	}
 
-	static constexpr bool bRecomputeTangents = false;
-	static constexpr bool bRecomputeNormals = false;
-	static constexpr bool bCacheOptimize = false;
-
+	FStaticToSkeletalMeshConverter::FInitializationParams Parameters;
+	Parameters.Materials = MaterialView;
+	Parameters.bRecomputeNormals = false;
+	Parameters.bRecomputeTangents = false;
+	Parameters.bCacheOptimize = false;
+	
 	if (!FStaticToSkeletalMeshConverter::InitializeSkeletalMeshFromMeshDescriptions(
-		SkeletalMesh, MeshDescriptions, MaterialView, ReferenceSkeleton, bRecomputeNormals, bRecomputeTangents, bCacheOptimize))
+		SkeletalMesh, MeshDescriptions, ReferenceSkeleton, Parameters))
 	{
 		return false;
 	}
diff --git a/Engine/Plugins/Runtime/MeshModelingToolset/Source/ModelingComponentsEditorOnly/Private/AssetUtils/CreateSkeletalMeshUtil.cpp b/Engine/Plugins/Runtime/MeshModelingToolset/Source/ModelingComponentsEditorOnly/Private/AssetUtils/CreateSkeletalMeshUtil.cpp
--- a/Engine/Plugins/Runtime/MeshModelingToolset/Source/ModelingComponentsEditorOnly/Private/AssetUtils/CreateSkeletalMeshUtil.cpp
+++ b/Engine/Plugins/Runtime/MeshModelingToolset/Source/ModelingComponentsEditorOnly/Private/AssetUtils/CreateSkeletalMeshUtil.cpp
@@ -151,10 +151,16 @@ UE::AssetUtils::ECreateSkeletalMeshResult UE::AssetUtils::CreateSkeletalMeshAsse
 		NewSkeletalMesh->NaniteSettings = Options.NaniteSettings;
 	}
 	
+	FStaticToSkeletalMeshConverter::FInitializationParams Parameters;
+	Parameters.Materials = MaterialView;
+	Parameters.bRecomputeNormals = Options.bEnableRecomputeNormals;
+	Parameters.bRecomputeTangents = Options.bEnableRecomputeTangents;
+	
 	if (!FStaticToSkeletalMeshConverter::InitializeSkeletalMeshFromMeshDescriptions(
-		NewSkeletalMesh, MeshDescriptions, MaterialView, 
+		NewSkeletalMesh,
+		MeshDescriptions, 
 		Options.RefSkeleton ? *Options.RefSkeleton : Options.Skeleton->GetReferenceSkeleton(),
-		Options.bEnableRecomputeNormals, Options.bEnableRecomputeTangents))
+		Parameters))
 	{
 		return ECreateSkeletalMeshResult::UnknownError;
 	}
diff --git a/Engine/Source/Developer/SkeletalMeshUtilitiesCommon/Private/StaticToSkeletalMeshConverter.cpp b/Engine/Source/Developer/SkeletalMeshUtilitiesCommon/Private/StaticToSkeletalMeshConverter.cpp
--- a/Engine/Source/Developer/SkeletalMeshUtilitiesCommon/Private/StaticToSkeletalMeshConverter.cpp
+++ b/Engine/Source/Developer/SkeletalMeshUtilitiesCommon/Private/StaticToSkeletalMeshConverter.cpp
@@ -506,6 +506,22 @@ bool FStaticToSkeletalMeshConverter::InitializeSkeletalMeshFromMeshDescriptions(
 	const bool bInRecomputeTangents,
 	const bool bCacheOptimize
 )	
+{
+	FInitializationParams Parameters;
+	Parameters.Materials = InMaterials;
+	Parameters.bRecomputeNormals = bInRecomputeNormals;
+	Parameters.bRecomputeTangents = bInRecomputeTangents;
+	Parameters.bCacheOptimize = bCacheOptimize;
+	
+	return InitializeSkeletalMeshFromMeshDescriptions(InSkeletalMesh, InMeshDescriptions, InReferenceSkeleton, Parameters);
+}
+
+bool FStaticToSkeletalMeshConverter::InitializeSkeletalMeshFromMeshDescriptions(
+	USkeletalMesh* InSkeletalMesh,
+	TArrayView<const FMeshDescription*> InMeshDescriptions,
+	const FReferenceSkeleton& InReferenceSkeleton,
+	const FInitializationParams& InParams
+)	
 {
 	if (!ensure(InSkeletalMesh))
 	{
@@ -535,12 +551,12 @@ bool FStaticToSkeletalMeshConverter::InitializeSkeletalMeshFromMeshDescriptions(
 
 	// Set the materials before we start converting. We'll add dummy materials afterward if there are more sections
 	// than materials in any of the LODs. Not the best system, but the best we have for now.
-	InSkeletalMesh->SetMaterials(TArray<FSkeletalMaterial>{InMaterials});
+	InSkeletalMesh->SetMaterials(TArray<FSkeletalMaterial>{InParams.Materials});
 
 	TSet<FName> ValidMaterialSlotNames;
-	for (int32 Index = 0; Index < InMaterials.Num(); Index++)
+	for (int32 Index = 0; Index < InParams.Materials.Num(); Index++)
 	{
-		const FSkeletalMaterial& Material = InMaterials[Index];
+		const FSkeletalMaterial& Material = InParams.Materials[Index];
 		if (!Material.MaterialSlotName.IsNone())
 		{
 			ValidMaterialSlotNames.Add(Material.MaterialSlotName);
@@ -568,8 +584,8 @@ bool FStaticToSkeletalMeshConverter::InitializeSkeletalMeshFromMeshDescriptions(
 			SkeletalLODInfo.ReductionSettings.MaxDeviationPercentage = 0.0f;
 			SkeletalLODInfo.LODHysteresis = 0.02f;
 
-			SkeletalLODInfo.BuildSettings.bRecomputeNormals = bInRecomputeNormals;
-			SkeletalLODInfo.BuildSettings.bRecomputeTangents = bInRecomputeTangents;
+			SkeletalLODInfo.BuildSettings.bRecomputeNormals = InParams.bRecomputeNormals;
+			SkeletalLODInfo.BuildSettings.bRecomputeTangents = InParams.bRecomputeTangents;
 
 			FMeshDescription ClonedDescription(*MeshDescription);
 
@@ -583,12 +599,33 @@ bool FStaticToSkeletalMeshConverter::InitializeSkeletalMeshFromMeshDescriptions(
 				if (!ValidMaterialSlotNames.Contains(MaterialSlotNamesAttribute.Get(PolygonGroupID)))
 				{
 					int32 MaterialIndex = PolygonGroupID.GetValue();
-					MaterialIndex = FMath::Clamp(MaterialIndex, 0, InMaterials.Num() - 1);
-					MaterialSlotNamesAttribute.Set(PolygonGroupID, InMaterials[MaterialIndex].MaterialSlotName);
+					MaterialIndex = FMath::Clamp(MaterialIndex, 0, InParams.Materials.Num() - 1);
+					if (InParams.Materials.IsValidIndex(MaterialIndex))
+					{
+						MaterialSlotNamesAttribute.Set(PolygonGroupID, InParams.Materials[MaterialIndex].MaterialSlotName);
+					}
 				}
 			}
 			
-			if (!AddLODFromMeshDescription(MoveTemp(ClonedDescription), InSkeletalMesh, MeshUtilities, bCacheOptimize))
+			// pre-add the vertex attributes to enable for render so that the render mesh generation takes them into account during creation. 
+			if (!InParams.VertexAttributesForRender.IsEmpty())
+			{
+				const TAttributesSet<FVertexID>& VertexAttributes = MeshDescription->VertexAttributes();
+				VertexAttributes.ForEachByType<float>([&InParams, &SkeletalLODInfo](const FName AttributeName, TVertexAttributesConstRef<float>)
+				{
+					if (!FSkeletalMeshAttributes::IsReservedAttributeName(AttributeName))
+					{
+						if (InParams.VertexAttributesForRender.Contains(AttributeName))
+						{
+							FSkeletalMeshVertexAttributeInfo& NewAttribute = SkeletalLODInfo.VertexAttributes.AddDefaulted_GetRef();
+							NewAttribute.Name = AttributeName;
+							NewAttribute.EnabledForRender = true;
+						}
+					}
+				});
+			}
+			
+			if (!AddLODFromMeshDescription(MoveTemp(ClonedDescription), InSkeletalMesh, MeshUtilities, InParams.bCacheOptimize))
 			{
 				// If we didn't get a model for LOD index 0, we don't have a mesh. Bail out.
 				if (bFirstSourceModel)
diff --git a/Engine/Source/Developer/SkeletalMeshUtilitiesCommon/Public/StaticToSkeletalMeshConverter.h b/Engine/Source/Developer/SkeletalMeshUtilitiesCommon/Public/StaticToSkeletalMeshConverter.h
--- a/Engine/Source/Developer/SkeletalMeshUtilitiesCommon/Public/StaticToSkeletalMeshConverter.h
+++ b/Engine/Source/Developer/SkeletalMeshUtilitiesCommon/Public/StaticToSkeletalMeshConverter.h
@@ -81,6 +81,7 @@ struct FStaticToSkeletalMeshConverter
 	 *  @param bInRecomputeTangents Recompute tangents when the mesh is generated from the mesh description.
 	 *  @param bCacheOptimize Enable the cache optimization of the index buffer.
 	 */
+	UE_DEPRECATED(5.8, "Please use the new overloads of InitializeSkeletalMeshFromMeshDescriptions that takes FInitializationParams.")
 	static SKELETALMESHUTILITIESCOMMON_API bool InitializeSkeletalMeshFromMeshDescriptions(
 		USkeletalMesh* InSkeletalMesh,
 		TArrayView<const FMeshDescription*> InMeshDescriptions,
@@ -90,6 +91,44 @@ struct FStaticToSkeletalMeshConverter
 		const bool bInRecomputeTangents = true,
 		const bool bCacheOptimize = true
 		);
+	
+	/* Structure that defines parameters passed into the initialize skeletal mesh from mesh descriptions function. */
+	struct FInitializationParams
+	{
+		FInitializationParams() {}
+
+		/** The materials to apply to the skeletal mesh. The conversion will attempt a best effort to
+	     * ensure that the PolygonGroup::ImportedMaterialSlotName attribute will map to existing import names.
+		 */
+		TConstArrayView<FSkeletalMaterial> Materials;
+		
+		/** Recompute normals when the mesh is generated from the mesh description. */
+		bool bRecomputeNormals = false;
+		
+		/** Recompute tangents when the mesh is generated from the mesh description. */
+		bool bRecomputeTangents = true;
+		
+		/** Enable the cache optimization of the index buffer. */
+		bool bCacheOptimize = true;
+		
+		/** Vertex attribute names, as defined in the mesh description, that need to be enabled for render. */
+		TConstArrayView<FName> VertexAttributesForRender;
+	};
+	
+	/** Create a skeletal mesh from a list of mesh description objects and a reference skeleton. If any mesh description
+	 *  either doesn't contain a skin weight, or any of the skin weights reference a bone that doesn't exist on the
+	 *  reference skeleton, this operation fails.
+	 *  @param InSkeletalMesh The skeletal mesh to initialize. It has to be freshly created and be completely empty.
+	 *  @param InMeshDescriptions The mesh description objects to create each subsequent LOD, starting from LOD0. 
+	 *  @param InReferenceSkeleton The reference skeleton to use.
+	 *  @param InParams The initialization parameters used.
+	 */
+	static SKELETALMESHUTILITIESCOMMON_API bool InitializeSkeletalMeshFromMeshDescriptions(
+		USkeletalMesh* InSkeletalMesh,
+		TArrayView<const FMeshDescription*> InMeshDescriptions,
+		const FReferenceSkeleton& InReferenceSkeleton,
+		const FInitializationParams& InParams = {}
+		);
 };
 
 #endif
diff --git a/Engine/Source/Runtime/Experimental/GeometryCollectionEngine/Private/GeometryCollection/GeometryCollectionEngineConversion.cpp b/Engine/Source/Runtime/Experimental/GeometryCollectionEngine/Private/GeometryCollection/GeometryCollectionEngineConversion.cpp
--- a/Engine/Source/Runtime/Experimental/GeometryCollectionEngine/Private/GeometryCollection/GeometryCollectionEngineConversion.cpp
+++ b/Engine/Source/Runtime/Experimental/GeometryCollectionEngine/Private/GeometryCollection/GeometryCollectionEngineConversion.cpp
@@ -2108,16 +2108,17 @@ void FGeometryCollectionEngineConversion::ConvertCollectionToSkeletalMesh(const
 
 	TArray<const FMeshDescription*> MeshDescriptions = { &MeshDescription };
 
-	// Update the skeletal mesh from the description 
-	constexpr bool bRecomputeNormals = false;
-	constexpr bool bRecomputeTangents = false;
+	// Update the skeletal mesh from the description
+	FStaticToSkeletalMeshConverter::FInitializationParams Parameters;
+	Parameters.Materials = SkeletalMaterials;
+	Parameters.bRecomputeNormals = false;
+	Parameters.bRecomputeTangents = false;
+	
 	FStaticToSkeletalMeshConverter::InitializeSkeletalMeshFromMeshDescriptions(
 		&OutSkeletalMesh,
 		MeshDescriptions,
-		SkeletalMaterials,
 		OutSkeleton.GetReferenceSkeleton(),
-		bRecomputeNormals,
-		bRecomputeTangents);
+		Parameters);
 
 	// finalize skeletal mesh binding to skeleton 
 	OutSkeletalMesh.SetSkeleton(&OutSkeleton);
