Skeletal Mesh: Search for orphaned morph targets using StaticFindObjectFastInternal to ensure we also find objects with the async flag set.

#rnx

diff --git a/Engine/Source/Developer/SkeletalMeshUtilitiesCommon/Private/LODUtilities.cpp b/Engine/Source/Developer/SkeletalMeshUtilitiesCommon/Private/LODUtilities.cpp
--- a/Engine/Source/Developer/SkeletalMeshUtilitiesCommon/Private/LODUtilities.cpp
+++ b/Engine/Source/Developer/SkeletalMeshUtilitiesCommon/Private/LODUtilities.cpp
@@ -3339,8 +3339,10 @@ static void BuildMorphTargetsInternal(
 				{
 					FGCScopeGuard Guard;
 
-					// Check if there's an orphaned morph target object floating about.
-					MorphTarget = Cast<UMorphTarget>(StaticFindObjectFast(UMorphTarget::StaticClass(), BaseSkelMesh, ObjectName));
+					// Check if there's an orphaned morph target object floating about. We use StaticFindObjectFastInternal
+					// rather than StaticFindObjectFast since we might have a situation where the orphaned morph target might
+					// still have the Async flag set, which StaticFindObjectFast will skip over.
+					MorphTarget = Cast<UMorphTarget>(StaticFindObjectFastInternal(UMorphTarget::StaticClass(), BaseSkelMesh, ObjectName, EFindObjectFlags::ExactClass));
 					if (MorphTarget)
 					{
 						// Temporarily set the async flag on it so that it doesn't vanish once we leave the GC scope.
